Set ReferencedGeometries and PerInstanceGeometries during FRHIRayTracingScene build instead of initializer.

diff --git a/Engine/Plugins/Experimental/GPULightmass/Source/GPULightmass/Private/LightmapRenderer.cpp b/Engine/Plugins/Experimental/GPULightmass/Source/GPULightmass/Private/LightmapRenderer.cpp
--- a/Engine/Plugins/Experimental/GPULightmass/Source/GPULightmass/Private/LightmapRenderer.cpp
+++ b/Engine/Plugins/Experimental/GPULightmass/Source/GPULightmass/Private/LightmapRenderer.cpp
@@ -1161,7 +1161,7 @@ bool FSceneRenderState::SetupRayTracingScene(FRDGBuilder& GraphBuilder, FSceneUn
 			AccelerationStructureAddressesBuffer.Initialize(
 				RHICmdList,
 				TEXT("LightmassRayTracingAccelerationStructureAddressesBuffer"),
-				SceneInitializer.ReferencedGeometries.Num() * sizeof(FRayTracingAccelerationStructureAddress),
+				SceneWithGeometryInstances.ReferencedGeometries.Num() * sizeof(FRayTracingAccelerationStructureAddress),
 				BUF_Volatile | BUF_MultiGPUAllocate);
 
 			const uint32 InstanceUploadBufferSize = SceneInitializer.NumNativeInstances * sizeof(FRayTracingInstanceDescriptorInput);
@@ -1206,13 +1206,15 @@ bool FSceneRenderState::SetupRayTracingScene(FRDGBuilder& GraphBuilder, FSceneUn
 					AccelerationStructureAddressesBuffer.Buffer,
 					GPUIndex,
 					0,
-					SceneInitializer.ReferencedGeometries.Num() * sizeof(FRayTracingAccelerationStructureAddress), RLM_WriteOnly);
+					SceneWithGeometryInstances.ReferencedGeometries.Num() * sizeof(FRayTracingAccelerationStructureAddress), RLM_WriteOnly);
 
-				RHICmdList.EnqueueLambda([AddressesPtr, &SceneInitializer, GPUIndex](FRHICommandListBase&)
+				const TArrayView<FRHIRayTracingGeometry*> ReferencedGeometries = RHICmdList.AllocArray(MakeConstArrayView(SceneWithGeometryInstances.ReferencedGeometries));
+
+				RHICmdList.EnqueueLambda([AddressesPtr, ReferencedGeometries, GPUIndex](FRHICommandListBase&)
 				{
-					for (int32 GeometryIndex = 0; GeometryIndex < SceneInitializer.ReferencedGeometries.Num(); ++GeometryIndex)
+					for (int32 GeometryIndex = 0; GeometryIndex < ReferencedGeometries.Num(); ++GeometryIndex)
 					{
-						AddressesPtr[GeometryIndex] = SceneInitializer.ReferencedGeometries[GeometryIndex]->GetAccelerationStructureAddress(GPUIndex);
+						AddressesPtr[GeometryIndex] = ReferencedGeometries[GeometryIndex]->GetAccelerationStructureAddress(GPUIndex);
 					}
 				});
 
@@ -1244,6 +1246,9 @@ bool FSceneRenderState::SetupRayTracingScene(FRDGBuilder& GraphBuilder, FSceneUn
 				BuildParams.ScratchBufferOffset = 0;
 				BuildParams.InstanceBuffer = InstanceBuffer.Buffer;
 				BuildParams.InstanceBufferOffset = 0;
+				BuildParams.ReferencedGeometries = SceneWithGeometryInstances.ReferencedGeometries;
+				BuildParams.PerInstanceGeometries = SceneWithGeometryInstances.PerInstanceGeometries;
+
 				RHICmdList.BuildAccelerationStructure(BuildParams);
 			}
 
diff --git a/Engine/Source/Runtime/D3D12RHI/Private/D3D12RayTracing.cpp b/Engine/Source/Runtime/D3D12RHI/Private/D3D12RayTracing.cpp
--- a/Engine/Source/Runtime/D3D12RHI/Private/D3D12RayTracing.cpp
+++ b/Engine/Source/Runtime/D3D12RHI/Private/D3D12RayTracing.cpp
@@ -3524,10 +3524,10 @@ void FD3D12RayTracingScene::BuildAccelerationStructure(FD3D12CommandContext& Com
 			PRAGMA_DISABLE_DEPRECATION_WARNINGS
 			checkf(ScratchBufferOffset + Layer.UpdateScratchOffset + PrebuildInfo.UpdateScratchDataSizeInBytes <= ScratchBuffer->GetSize(),
 				TEXT("TLAS scratch buffer size is %d bytes with offset %d (%d bytes available), but the update requires %lld bytes. ")
-				TEXT("BuildInputs.NumDescs = %d, Instances.Num = %d, PerInstanceGeometries.Num = %d."),
+				TEXT("BuildInputs.NumDescs = %d, Instances.Num = %d."),
 				ScratchBuffer->GetSize(), ScratchBufferOffset + Layer.UpdateScratchOffset, ScratchBuffer->GetSize() - ScratchBufferOffset,
 				PrebuildInfo.UpdateScratchDataSizeInBytes,
-				Layer.BuildInputs.NumDescs, Initializer.NumNativeInstancesPerLayer[LayerIndex], Initializer.PerInstanceGeometries.Num());
+				Layer.BuildInputs.NumDescs, Initializer.NumNativeInstancesPerLayer[LayerIndex]);
 			PRAGMA_ENABLE_DEPRECATION_WARNINGS
 		}
 		else
@@ -3535,10 +3535,10 @@ void FD3D12RayTracingScene::BuildAccelerationStructure(FD3D12CommandContext& Com
 			PRAGMA_DISABLE_DEPRECATION_WARNINGS
 			checkf(ScratchBufferOffset + Layer.BuildScratchOffset + PrebuildInfo.ScratchDataSizeInBytes <= ScratchBuffer->GetSize(),
 				TEXT("TLAS scratch buffer size is %d bytes with offset %d (%d bytes available), but the build requires %lld bytes. ")
-				TEXT("BuildInputs.NumDescs = %d, Instances.Num = %d, PerInstanceGeometries.Num = %d."),
+				TEXT("BuildInputs.NumDescs = %d, Instances.Num = %d."),
 				ScratchBuffer->GetSize(), ScratchBufferOffset + Layer.BuildScratchOffset, ScratchBuffer->GetSize() - ScratchBufferOffset,
 				PrebuildInfo.ScratchDataSizeInBytes,
-				Layer.BuildInputs.NumDescs, Initializer.NumNativeInstancesPerLayer[LayerIndex], Initializer.PerInstanceGeometries.Num());
+				Layer.BuildInputs.NumDescs, Initializer.NumNativeInstancesPerLayer[LayerIndex]);
 			PRAGMA_ENABLE_DEPRECATION_WARNINGS
 		}
 	}
@@ -3591,10 +3591,10 @@ void FD3D12RayTracingScene::BuildAccelerationStructure(FD3D12CommandContext& Com
 			#endif // ENABLE_RESIDENCY_MANAGEMENT
 			};
 
-			const int32 NumReferencedGeometries = Initializer.ReferencedGeometries.Num();
+			const int32 NumReferencedGeometries = ReferencedGeometries.Num();
 			for (int32 Index = 0; Index < NumReferencedGeometries; ++Index)
 			{
-				FD3D12RayTracingGeometry* Geometry = FD3D12DynamicRHI::ResourceCast(Initializer.ReferencedGeometries[Index].GetReference());
+				FD3D12RayTracingGeometry* Geometry = FD3D12DynamicRHI::ResourceCast(ReferencedGeometries[Index].GetReference());
 
 				checkf(!Geometry->IsDirty(CommandContext.GetGPUIndex()),
 					TEXT("Acceleration structures for all geometries must be built before building the top level acceleration structure for the scene."));
@@ -3623,10 +3623,10 @@ void FD3D12RayTracingScene::BuildAccelerationStructure(FD3D12CommandContext& Com
 	
 	// Build the actual acceleration structure
 
-	const int32 NumReferencedGeometries = Initializer.ReferencedGeometries.Num();
+	const int32 NumReferencedGeometries = ReferencedGeometries.Num();
 	for (int32 Index = 0; Index < NumReferencedGeometries; ++Index)
 	{
-		FD3D12RayTracingGeometry* Geometry = FD3D12DynamicRHI::ResourceCast(Initializer.ReferencedGeometries[Index].GetReference());
+		FD3D12RayTracingGeometry* Geometry = FD3D12DynamicRHI::ResourceCast(ReferencedGeometries[Index].GetReference());
 		CommandContext.UpdateResidency(Geometry->AccelerationStructureBuffers[GPUIndex]->ResourceLocation.GetResource());
 	}
 
@@ -3985,6 +3985,14 @@ void FD3D12CommandContext::RHIBuildAccelerationStructure(const FRayTracingSceneB
 	FD3D12RayTracingScene* Scene = FD3D12DynamicRHI::ResourceCast(SceneBuildParams.Scene);
 	FD3D12Buffer* ScratchBuffer = RetrieveObject<FD3D12Buffer>(SceneBuildParams.ScratchBuffer);
 	FD3D12Buffer* InstanceBuffer = RetrieveObject<FD3D12Buffer>(SceneBuildParams.InstanceBuffer);
+
+	Scene->ReferencedGeometries.Reserve(SceneBuildParams.ReferencedGeometries.Num());
+
+	for (FRHIRayTracingGeometry* ReferencedGeometry : SceneBuildParams.ReferencedGeometries)
+	{
+		Scene->ReferencedGeometries.Add(ReferencedGeometry);
+	}
+
 	Scene->BuildAccelerationStructure(
 		*this,
 		ScratchBuffer, SceneBuildParams.ScratchBufferOffset,
diff --git a/Engine/Source/Runtime/D3D12RHI/Private/D3D12RayTracing.h b/Engine/Source/Runtime/D3D12RHI/Private/D3D12RayTracing.h
--- a/Engine/Source/Runtime/D3D12RHI/Private/D3D12RayTracing.h
+++ b/Engine/Source/Runtime/D3D12RHI/Private/D3D12RayTracing.h
@@ -126,6 +126,10 @@ public:
 
 	const FRayTracingSceneInitializer2 Initializer;
 
+	// Unique list of geometries referenced by all instances in this scene.
+	// Any referenced geometry is kept alive while the scene is alive.
+	TArray<TRefCountPtr<FRHIRayTracingGeometry>> ReferencedGeometries;
+
 	// Scene keeps track of child acceleration structure buffers to ensure
 	// they are resident when any ray tracing work is dispatched.
 	// Resources that share residency handles are deduplicated.
diff --git a/Engine/Source/Runtime/D3D12RHI/Private/D3D12RayTracingDebug.cpp b/Engine/Source/Runtime/D3D12RHI/Private/D3D12RayTracingDebug.cpp
--- a/Engine/Source/Runtime/D3D12RHI/Private/D3D12RayTracingDebug.cpp
+++ b/Engine/Source/Runtime/D3D12RHI/Private/D3D12RayTracingDebug.cpp
@@ -149,10 +149,10 @@ static void DebugSerializeScene(const FD3D12RayTracingScene& Scene, FD3D12Buffer
 	checkf(OutputBuffer.Num() == 0, TEXT("Scene header must be written into the buffer first"));
 	Serialize(&SceneHeader, sizeof(SceneHeader), EAlignmentMode::AlignBeginning); // Reserve space for the header
 
-	const int32 NumReferencedGeometries = SceneInitializer.ReferencedGeometries.Num();
+	const int32 NumReferencedGeometries = Scene.ReferencedGeometries.Num();
 	for (int32 GeometryIndex = 0; GeometryIndex < NumReferencedGeometries; ++GeometryIndex)
 	{
-		FD3D12RayTracingGeometry* Geometry = FD3D12DynamicRHI::ResourceCast(SceneInitializer.ReferencedGeometries[GeometryIndex].GetReference());
+		FD3D12RayTracingGeometry* Geometry = FD3D12DynamicRHI::ResourceCast(Scene.ReferencedGeometries[GeometryIndex].GetReference());
 		TRefCountPtr<FD3D12Buffer> BlasBuffer = Geometry->AccelerationStructureBuffers[0];
 		D3D12_GPU_VIRTUAL_ADDRESS Address = BlasBuffer->ResourceLocation.GetGPUVirtualAddress();
 		GeometryMap.FindOrAdd(Address, GeometryMap.Num());
@@ -223,7 +223,7 @@ static void DebugSerializeScene(const FD3D12RayTracingScene& Scene, FD3D12Buffer
 		{
 			FGeometryHeader GeometryHeader;
 
-			FD3D12RayTracingGeometry* Geometry = FD3D12DynamicRHI::ResourceCast(SceneInitializer.ReferencedGeometries[GeometryIndex].GetReference());
+			FD3D12RayTracingGeometry* Geometry = FD3D12DynamicRHI::ResourceCast(Scene.ReferencedGeometries[GeometryIndex].GetReference());
 
 			FString DebugName = Geometry->DebugName.ToString();
 			FString OwnerName = Geometry->OwnerName.ToString();
diff --git a/Engine/Source/Runtime/RHI/Public/RHICommandList.h b/Engine/Source/Runtime/RHI/Public/RHICommandList.h
--- a/Engine/Source/Runtime/RHI/Public/RHICommandList.h
+++ b/Engine/Source/Runtime/RHI/Public/RHICommandList.h
@@ -2424,8 +2424,8 @@ FRHICOMMAND_UNNAMED(FRHICommandBuildAccelerationStructure)
 {
 	FRayTracingSceneBuildParams SceneBuildParams;
 
-	explicit FRHICommandBuildAccelerationStructure(const FRayTracingSceneBuildParams& InSceneBuildParams)
-		: SceneBuildParams(InSceneBuildParams)
+	explicit FRHICommandBuildAccelerationStructure(FRayTracingSceneBuildParams InSceneBuildParams)
+		: SceneBuildParams(MoveTemp(InSceneBuildParams))
 	{}
 
 	RHI_API void Execute(FRHICommandListBase& CmdList);
@@ -3394,7 +3394,15 @@ public:
 		}
 		else
 		{
-			ALLOC_COMMAND(FRHICommandBuildAccelerationStructure)(SceneBuildParams);
+			FRayTracingSceneBuildParams InlineParams = SceneBuildParams;
+			InlineParams.ReferencedGeometries = AllocArray(SceneBuildParams.ReferencedGeometries);
+			InlineParams.PerInstanceGeometries = AllocArray(SceneBuildParams.PerInstanceGeometries);
+
+			ALLOC_COMMAND(FRHICommandBuildAccelerationStructure)(MoveTemp(InlineParams));
+
+			// This RHI command modifies members of the FRHIRayTracingScene inside platform RHI implementations.
+			// It therefore needs the RHI lock fence to prevent races on those members.
+			RHIThreadFence(true);
 		}
 	}
 
@@ -4288,7 +4296,9 @@ public:
 		check(LooseParameterDataSize <= UINT16_MAX);
 
 		FRayTracingLocalShaderBindings* InlineBindings = new(Alloc<FRayTracingLocalShaderBindings>()) FRayTracingLocalShaderBindings();
+		PRAGMA_DISABLE_DEPRECATION_WARNINGS
 		InlineBindings->Geometry = Scene->GetInitializer().PerInstanceGeometries[InstanceIndex];
+		PRAGMA_ENABLE_DEPRECATION_WARNINGS
 		InlineBindings->SegmentIndex = SegmentIndex;
 		PRAGMA_DISABLE_DEPRECATION_WARNINGS
 		InlineBindings->RecordIndex = (Scene->GetInitializer().SegmentPrefixSum[InstanceIndex] + SegmentIndex) * Scene->GetInitializer().ShaderSlotsPerGeometrySegment + ShaderSlot;
diff --git a/Engine/Source/Runtime/RHI/Public/RHIContext.h b/Engine/Source/Runtime/RHI/Public/RHIContext.h
--- a/Engine/Source/Runtime/RHI/Public/RHIContext.h
+++ b/Engine/Source/Runtime/RHI/Public/RHIContext.h
@@ -560,6 +560,12 @@ struct FRayTracingSceneBuildParams
 	FRHIBuffer* InstanceBuffer = nullptr;
 	uint32 InstanceBufferOffset = 0;
 
+	// Unique list of geometries referenced by all instances in this scene.
+	// Any referenced geometry is kept alive while the scene is alive.
+	TConstArrayView<FRHIRayTracingGeometry*> ReferencedGeometries;
+	// One entry per instance
+	TConstArrayView<FRHIRayTracingGeometry*> PerInstanceGeometries;
+
 	EAccelerationStructureBuildMode BuildMode = EAccelerationStructureBuildMode::Build;
 };
 
diff --git a/Engine/Source/Runtime/RHI/Public/RHIResources.h b/Engine/Source/Runtime/RHI/Public/RHIResources.h
--- a/Engine/Source/Runtime/RHI/Public/RHIResources.h
+++ b/Engine/Source/Runtime/RHI/Public/RHIResources.h
@@ -3284,9 +3284,12 @@ struct FRayTracingSceneInitializer2
 {
 	// Unique list of geometries referenced by all instances in this scene.
 	// Any referenced geometry is kept alive while the scene is alive.
+	UE_DEPRECATED(5.5, "ReferencedGeometries should be provided in FRayTracingSceneBuildParams.")
 	TArray<TRefCountPtr<FRHIRayTracingGeometry>> ReferencedGeometries;
 	// One entry per instance
+	UE_DEPRECATED(5.5, "PerInstanceGeometries should be provided in FRayTracingSceneBuildParams.")
 	TArray<FRHIRayTracingGeometry*> PerInstanceGeometries;
+
 	// Exclusive prefix sum of `Instance.NumTransforms` for all instances in this scene. Used to emulate SV_InstanceID in hit shaders.
 	UE_DEPRECATED(5.5, "Providing BaseInstancePrefixSum is no longer necessary.")
 	TArray<uint32> BaseInstancePrefixSum;
diff --git a/Engine/Source/Runtime/Renderer/Private/RayTracing/RayTracingInstanceBufferUtil.cpp b/Engine/Source/Runtime/Renderer/Private/RayTracing/RayTracingInstanceBufferUtil.cpp
--- a/Engine/Source/Runtime/Renderer/Private/RayTracing/RayTracingInstanceBufferUtil.cpp
+++ b/Engine/Source/Runtime/Renderer/Private/RayTracing/RayTracingInstanceBufferUtil.cpp
@@ -42,10 +42,10 @@ FRayTracingSceneWithGeometryInstances CreateRayTracingSceneWithGeometryInstances
 	Output.InstanceGeometryIndices.SetNumUninitialized(NumSceneInstances);
 	Output.BaseUploadBufferOffsets.SetNumUninitialized(NumSceneInstances);
 	Output.BaseInstancePrefixSum.SetNumUninitialized(NumSceneInstances);
+	Output.PerInstanceGeometries.SetNumUninitialized(NumSceneInstances);
 
 	FRayTracingSceneInitializer2 Initializer;
 	Initializer.DebugName = FName(TEXT("FRayTracingScene"));
-	Initializer.PerInstanceGeometries.SetNumUninitialized(NumSceneInstances);
 	Initializer.BuildFlags = BuildFlags;
 
 	PRAGMA_DISABLE_DEPRECATION_WARNINGS
@@ -77,7 +77,7 @@ FRayTracingSceneWithGeometryInstances CreateRayTracingSceneWithGeometryInstances
 
 		checkf(InstanceDesc.GeometryRHI, TEXT("Ray tracing instance must have a valid geometry."));
 
-		Initializer.PerInstanceGeometries[InstanceIndex] = InstanceDesc.GeometryRHI;
+		Output.PerInstanceGeometries[InstanceIndex] = InstanceDesc.GeometryRHI;
 
 		PRAGMA_DISABLE_DEPRECATION_WARNINGS
 		// Compute geometry segment count prefix sum to be later used in GetHitRecordBaseIndex()
@@ -86,11 +86,11 @@ FRayTracingSceneWithGeometryInstances CreateRayTracingSceneWithGeometryInstances
 
 		Output.TotalNumSegments += InstanceDesc.GeometryRHI->GetNumSegments();
 
-		uint32 GeometryIndex = UniqueGeometries.FindOrAdd(InstanceDesc.GeometryRHI, Initializer.ReferencedGeometries.Num());
+		uint32 GeometryIndex = UniqueGeometries.FindOrAdd(InstanceDesc.GeometryRHI, Output.ReferencedGeometries.Num());
 		Output.InstanceGeometryIndices[InstanceIndex] = GeometryIndex;
-		if (GeometryIndex == Initializer.ReferencedGeometries.Num())
+		if (GeometryIndex == Output.ReferencedGeometries.Num())
 		{
-			Initializer.ReferencedGeometries.Add(InstanceDesc.GeometryRHI);
+			Output.ReferencedGeometries.Add(InstanceDesc.GeometryRHI);
 		}
 
 		if (bGpuSceneInstance)
diff --git a/Engine/Source/Runtime/Renderer/Private/RayTracing/RayTracingScene.cpp b/Engine/Source/Runtime/Renderer/Private/RayTracing/RayTracingScene.cpp
--- a/Engine/Source/Runtime/Renderer/Private/RayTracing/RayTracingScene.cpp
+++ b/Engine/Source/Runtime/Renderer/Private/RayTracing/RayTracingScene.cpp
@@ -94,8 +94,6 @@ void FRayTracingScene::Create(FRDGBuilder& GraphBuilder, const FViewInfo& View,
 
 		Layer.RayTracingSceneRHI = Layer.InitializationData.Scene;
 
-		const FRayTracingSceneInitializer2& SceneInitializer = Layer.RayTracingSceneRHI->GetInitializer();
-
 		const uint32 NumNativeInstances = Layer.InitializationData.NumNativeGPUSceneInstances + Layer.InitializationData.NumNativeCPUInstances;
 		const uint32 NumNativeInstancesAligned = FMath::DivideAndRoundUp(FMath::Max(NumNativeInstances, 1U), AllocationGranularity) * AllocationGranularity;
 		const uint32 NumTransformsAligned = FMath::DivideAndRoundUp(FMath::Max(Layer.InitializationData.NumNativeCPUInstances, 1U), AllocationGranularity) * AllocationGranularity;
@@ -138,7 +136,7 @@ void FRayTracingScene::Create(FRDGBuilder& GraphBuilder, const FViewInfo& View,
 
 		{
 			// Round to PoT to avoid resizing too often
-			const uint32 NumGeometries = FMath::RoundUpToPowerOfTwo(SceneInitializer.ReferencedGeometries.Num());
+			const uint32 NumGeometries = FMath::RoundUpToPowerOfTwo(Layer.InitializationData.ReferencedGeometries.Num());
 			const uint32 AccelerationStructureAddressesBufferSize = NumGeometries * sizeof(FRayTracingAccelerationStructureAddress);
 
 			if (Layer.AccelerationStructureAddressesBuffer.NumBytes < AccelerationStructureAddressesBufferSize)
@@ -293,7 +291,7 @@ void FRayTracingScene::Create(FRDGBuilder& GraphBuilder, const FViewInfo& View,
 					}
 				});
 
-			GraphBuilder.AddCommandListSetupTask([&Layer, &SceneInitializer](FRHICommandList& RHICmdList)
+			GraphBuilder.AddCommandListSetupTask([&Layer](FRHICommandList& RHICmdList)
 				{
 					for (uint32 GPUIndex : RHICmdList.GetGPUMask())
 					{
@@ -301,15 +299,17 @@ void FRayTracingScene::Create(FRDGBuilder& GraphBuilder, const FViewInfo& View,
 							Layer.AccelerationStructureAddressesBuffer.Buffer,
 							GPUIndex,
 							0,
-							SceneInitializer.ReferencedGeometries.Num() * sizeof(FRayTracingAccelerationStructureAddress), RLM_WriteOnly);
+							Layer.InitializationData.ReferencedGeometries.Num() * sizeof(FRayTracingAccelerationStructureAddress), RLM_WriteOnly);
+
+						const TArrayView<FRHIRayTracingGeometry*> ReferencedGeometries = RHICmdList.AllocArray(MakeConstArrayView(Layer.InitializationData.ReferencedGeometries));
 
-						RHICmdList.EnqueueLambda([AddressesPtr, &SceneInitializer, GPUIndex](FRHICommandListBase&)
+						RHICmdList.EnqueueLambda([AddressesPtr, ReferencedGeometries, GPUIndex](FRHICommandListBase&)
 							{
 								TRACE_CPUPROFILER_EVENT_SCOPE(GetAccelerationStructuresAddresses);
 
-								for (int32 GeometryIndex = 0; GeometryIndex < SceneInitializer.ReferencedGeometries.Num(); ++GeometryIndex)
+								for (int32 GeometryIndex = 0; GeometryIndex < ReferencedGeometries.Num(); ++GeometryIndex)
 								{
-									AddressesPtr[GeometryIndex] = SceneInitializer.ReferencedGeometries[GeometryIndex]->GetAccelerationStructureAddress(GPUIndex);
+									AddressesPtr[GeometryIndex] = ReferencedGeometries[GeometryIndex]->GetAccelerationStructureAddress(GPUIndex);
 								}
 							});
 
@@ -452,6 +452,8 @@ void FRayTracingScene::Build(FRDGBuilder& GraphBuilder, ERDGPassFlags ComputePas
 				BuildParams.ScratchBufferOffset = 0;
 				BuildParams.InstanceBuffer = PassParams->InstanceBuffer->GetRHI();
 				BuildParams.InstanceBufferOffset = 0;
+				BuildParams.ReferencedGeometries = Layer.InitializationData.ReferencedGeometries;
+				BuildParams.PerInstanceGeometries = Layer.InitializationData.PerInstanceGeometries;
 
 				RHICmdList.BindAccelerationStructureMemory(Layer.RayTracingSceneRHI, PassParams->TLASBuffer->GetRHI(), 0);
 				RHICmdList.BuildAccelerationStructure(BuildParams);
diff --git a/Engine/Source/Runtime/Renderer/Private/Tests/RayTracingTestbed.cpp b/Engine/Source/Runtime/Renderer/Private/Tests/RayTracingTestbed.cpp
--- a/Engine/Source/Runtime/Renderer/Private/Tests/RayTracingTestbed.cpp
+++ b/Engine/Source/Runtime/Renderer/Private/Tests/RayTracingTestbed.cpp
@@ -248,13 +248,15 @@ bool RunRayTracingTestbed_RenderThread(const FString& Parameters)
 			AccelerationStructureAddressesBuffer.Buffer,
 			GPUIndex,
 			0,
-			SceneInitializer.ReferencedGeometries.Num() * sizeof(FRayTracingAccelerationStructureAddress), RLM_WriteOnly);
+			RayTracingScene.ReferencedGeometries.Num() * sizeof(FRayTracingAccelerationStructureAddress), RLM_WriteOnly);
 
-		RHICmdList.EnqueueLambda([AddressesPtr, &SceneInitializer, GPUIndex](FRHICommandListBase&)
+		const TArrayView<FRHIRayTracingGeometry*> ReferencedGeometries = RHICmdList.AllocArray(MakeConstArrayView(RayTracingScene.ReferencedGeometries));
+
+		RHICmdList.EnqueueLambda([AddressesPtr, ReferencedGeometries, GPUIndex](FRHICommandListBase&)
 		{
-			for (int32 GeometryIndex = 0; GeometryIndex < SceneInitializer.ReferencedGeometries.Num(); ++GeometryIndex)
+			for (int32 GeometryIndex = 0; GeometryIndex < ReferencedGeometries.Num(); ++GeometryIndex)
 			{
-				AddressesPtr[GeometryIndex] = SceneInitializer.ReferencedGeometries[GeometryIndex]->GetAccelerationStructureAddress(GPUIndex);
+				AddressesPtr[GeometryIndex] = ReferencedGeometries[GeometryIndex]->GetAccelerationStructureAddress(GPUIndex);
 			}
 		});
 
@@ -288,6 +290,8 @@ bool RunRayTracingTestbed_RenderThread(const FString& Parameters)
 	BuildParams.ScratchBufferOffset = 0;
 	BuildParams.InstanceBuffer = InstanceBuffer.Buffer;
 	BuildParams.InstanceBufferOffset = 0;
+	BuildParams.ReferencedGeometries = RayTracingScene.ReferencedGeometries;
+	BuildParams.PerInstanceGeometries = RayTracingScene.PerInstanceGeometries;
 
 	RHICmdList.BuildAccelerationStructure(BuildParams);
 
diff --git a/Engine/Source/Runtime/Renderer/Public/RayTracingInstanceBufferUtil.h b/Engine/Source/Runtime/Renderer/Public/RayTracingInstanceBufferUtil.h
--- a/Engine/Source/Runtime/Renderer/Public/RayTracingInstanceBufferUtil.h
+++ b/Engine/Source/Runtime/Renderer/Public/RayTracingInstanceBufferUtil.h
@@ -58,6 +58,12 @@ struct FRayTracingSceneWithGeometryInstances
 	UE_DEPRECATED(5.5, "GPUInstances no longer supported. Use GPUSceneInstances instead.")
 	TArray<FRayTracingGPUInstance> GPUInstances;
 
+	// Unique list of geometries referenced by all instances in this scene.
+	// Any referenced geometry is kept alive while the scene is alive.
+	TArray<FRHIRayTracingGeometry*> ReferencedGeometries;
+	// One entry per instance
+	TArray<FRHIRayTracingGeometry*> PerInstanceGeometries;
+
 	PRAGMA_DISABLE_DEPRECATION_WARNINGS
 	FRayTracingSceneWithGeometryInstances() = default;
 	FRayTracingSceneWithGeometryInstances(const FRayTracingSceneWithGeometryInstances&) = default;
diff --git a/Engine/Source/Runtime/VulkanRHI/Private/VulkanRayTracing.cpp b/Engine/Source/Runtime/VulkanRHI/Private/VulkanRayTracing.cpp
--- a/Engine/Source/Runtime/VulkanRHI/Private/VulkanRayTracing.cpp
+++ b/Engine/Source/Runtime/VulkanRHI/Private/VulkanRayTracing.cpp
@@ -883,7 +883,7 @@ void FVulkanRayTracingScene::BuildPerInstanceGeometryParameterBuffer(FRHICommand
 	FVulkanRayTracingGeometryParameters* MappedParameters = reinterpret_cast<FVulkanRayTracingGeometryParameters*>(MappedBuffer);
 	uint32 ParameterIndex = 0;
 
-	for (FRHIRayTracingGeometry* GeometryRHI : Initializer.PerInstanceGeometries)
+	for (FRHIRayTracingGeometry* GeometryRHI : PerInstanceGeometries)
 	{
 		const FVulkanRayTracingGeometry* Geometry = ResourceCast(GeometryRHI);
 		const FRayTracingGeometryInitializer& GeometryInitializer = Geometry->GetInitializer();
@@ -1367,6 +1367,16 @@ void FVulkanCommandListContext::RHIBuildAccelerationStructure(const FRayTracingS
 	FVulkanRayTracingScene* const Scene = ResourceCast(SceneBuildParams.Scene);
 	FVulkanResourceMultiBuffer* const ScratchBuffer = ResourceCast(SceneBuildParams.ScratchBuffer);
 	FVulkanResourceMultiBuffer* const InstanceBuffer = ResourceCast(SceneBuildParams.InstanceBuffer);
+
+	Scene->ReferencedGeometries.Reserve(SceneBuildParams.ReferencedGeometries.Num());
+
+	for (FRHIRayTracingGeometry* ReferencedGeometry : SceneBuildParams.ReferencedGeometries)
+	{
+		Scene->ReferencedGeometries.Add(ReferencedGeometry);
+	}
+
+	Scene->PerInstanceGeometries = SceneBuildParams.PerInstanceGeometries;
+
 	Scene->BuildAccelerationStructure(
 		*this,
 		ScratchBuffer, SceneBuildParams.ScratchBufferOffset, 
diff --git a/Engine/Source/Runtime/VulkanRHI/Private/VulkanRayTracing.h b/Engine/Source/Runtime/VulkanRHI/Private/VulkanRayTracing.h
--- a/Engine/Source/Runtime/VulkanRHI/Private/VulkanRayTracing.h
+++ b/Engine/Source/Runtime/VulkanRHI/Private/VulkanRayTracing.h
@@ -239,6 +239,12 @@ public:
 private:
 	const FRayTracingSceneInitializer2 Initializer;
 
+	// Unique list of geometries referenced by all instances in this scene.
+	// Any referenced geometry is kept alive while the scene is alive.
+	TArray<TRefCountPtr<FRHIRayTracingGeometry>> ReferencedGeometries;
+	// One entry per instance
+	TArray<FRHIRayTracingGeometry*> PerInstanceGeometries;
+
 	// Native TLAS handles are owned by SRV objects in Vulkan RHI.
 	// D3D12 and other RHIs allow creating TLAS SRVs from any GPU address at any point
 	// and do not require them for operations such as build or update.
