MGPU:  Fixes for startup crashes in ray tracing and async texture creation.

#rnx

diff --git a/Engine/Source/Runtime/D3D12RHI/Private/D3D12RHI.cpp b/Engine/Source/Runtime/D3D12RHI/Private/D3D12RHI.cpp
--- a/Engine/Source/Runtime/D3D12RHI/Private/D3D12RHI.cpp
+++ b/Engine/Source/Runtime/D3D12RHI/Private/D3D12RHI.cpp
@@ -565,7 +565,7 @@ void FD3D12DynamicRHI::RHIEndFrame_RenderThread(FRHICommandListImmediate& RHICmd
 
 #if D3D12_RHI_RAYTRACING
 				// @todo dev-pr - explicit use of graphics context - nothing is synchronizing async compute - needs refactor
-				Device->GetRayTracingCompactionRequestHandler()->Update(*static_cast<FD3D12CommandContext*>(Contexts[ERHIPipeline::Graphics]));
+				Device->GetRayTracingCompactionRequestHandler()->Update(*Contexts[ERHIPipeline::Graphics]->GetSingleDeviceContext(Device->GetGPUIndex()));
 #endif // D3D12_RHI_RAYTRACING
 
 #if (RHI_NEW_GPU_PROFILER == 0)
diff --git a/Engine/Source/Runtime/D3D12RHI/Private/D3D12Texture.cpp b/Engine/Source/Runtime/D3D12RHI/Private/D3D12Texture.cpp
--- a/Engine/Source/Runtime/D3D12RHI/Private/D3D12Texture.cpp
+++ b/Engine/Source/Runtime/D3D12RHI/Private/D3D12Texture.cpp
@@ -980,6 +980,9 @@ FTextureRHIRef FD3D12DynamicRHI::RHIAsyncCreateTexture2D(uint32 SizeX, uint32 Si
 
 		for (FD3D12Texture& CurrentTexture : *TextureOut)
 		{
+			// Need to get device from GPU specific copy of the texture
+			Device = CurrentTexture.GetParentDevice();
+
 			FD3D12Resource* Resource = CurrentTexture.GetResource();
 
 			FD3D12SyncPointRef SyncPoint;
