Fix static analysis warnings

#rnx

diff --git a/Engine/Plugins/Animation/ControlRig/Source/ControlRig/Private/ControlRig.cpp b/Engine/Plugins/Animation/ControlRig/Source/ControlRig/Private/ControlRig.cpp
--- a/Engine/Plugins/Animation/ControlRig/Source/ControlRig/Private/ControlRig.cpp
+++ b/Engine/Plugins/Animation/ControlRig/Source/ControlRig/Private/ControlRig.cpp
@@ -1718,8 +1718,7 @@ TArray<UControlRig*> UControlRig::FindControlRigs(UObject* Outer, TSubclassOf<UC
 	for (TObjectIterator<UControlRig> Itr; Itr; ++Itr)
 	{
 		UControlRig* RigInstance = *Itr;
-		const UClass* RigInstanceClass = RigInstance ? RigInstance->GetClass() : nullptr;
-		if (OptionalClass == nullptr || (RigInstanceClass && RigInstanceClass->IsChildOf(OptionalClass)))
+		if (OptionalClass == nullptr || RigInstance->GetClass()->IsChildOf(OptionalClass))
 		{
 			if(RigInstance->IsInOuter(Outer))
 			{
diff --git a/Engine/Source/Runtime/AudioMixer/Private/AudioMixerSourceDecode.cpp b/Engine/Source/Runtime/AudioMixer/Private/AudioMixerSourceDecode.cpp
--- a/Engine/Source/Runtime/AudioMixer/Private/AudioMixerSourceDecode.cpp
+++ b/Engine/Source/Runtime/AudioMixer/Private/AudioMixerSourceDecode.cpp
@@ -161,6 +161,7 @@ public:
 				TArray<uint8> DecodeBuffer;
 				DecodeBuffer.AddZeroed(ByteSize);
 
+#if PLATFORM_NUM_AUDIODECOMPRESSION_PRECACHE_BUFFERS
 				// skip the first buffers if we've already decoded them during Precache:
 				if (DecodeTaskData.bSkipFirstBuffer)
 				{
@@ -181,6 +182,7 @@ public:
 						}
 					}
 				}
+#endif
 
 				const int32 kPCMBufferSize = NumChannels * DecodeTaskData.NumFramesToDecode * sizeof(int16);
 				int32 NumBytesStreamed = kPCMBufferSize;
diff --git a/Engine/Source/Runtime/CoreUObject/Private/Serialization/AsyncLoading.cpp b/Engine/Source/Runtime/CoreUObject/Private/Serialization/AsyncLoading.cpp
--- a/Engine/Source/Runtime/CoreUObject/Private/Serialization/AsyncLoading.cpp
+++ b/Engine/Source/Runtime/CoreUObject/Private/Serialization/AsyncLoading.cpp
@@ -380,6 +380,7 @@ struct FEDLBootNotificationManager
 
 				if (!bAnyParentOnStack)
 				{
+					checkSlow(Class);
 					CDORecursiveStack.Push(Class);
 					//FPlatformMisc::LowLevelOutputDebugStringf(TEXT("Create CDO %s\r\n"), *BootObject->GetName());
 					Class->GetDefaultObject();
diff --git a/Engine/Source/Runtime/D3D12RHI/Private/D3D12RayTracing.cpp b/Engine/Source/Runtime/D3D12RHI/Private/D3D12RayTracing.cpp
--- a/Engine/Source/Runtime/D3D12RHI/Private/D3D12RayTracing.cpp
+++ b/Engine/Source/Runtime/D3D12RHI/Private/D3D12RayTracing.cpp
@@ -86,6 +86,16 @@ static FAutoConsoleVariableRef CVarD3D12RayTracingGPUValidation(
 	TEXT("Whether to perform validation of ray tracing geometry and other structures on the GPU. Requires Shader Model 6. (default = 0)")
 );
 
+// This is required to avoid redundent code static analysis warnings
+// If the static_assert fires the assumptions it is predicated on have been
+// violated and the code should be revisited
+#if WITH_MGPU
+#define FOREACH_GPU(Condition, Function) for (uint32 GPUIndex = 0; Condition; ++GPUIndex) Function
+#else
+static_assert(MAX_NUM_GPUS == 1 && GNumExplicitGPUsForRendering == 1);
+#define FOREACH_GPU(Condition, Function) { constexpr uint32 GPUIndex = 0; Function }
+#endif
+
 // Ray tracing stat counters
 
 DECLARE_STATS_GROUP(TEXT("D3D12RHI: Ray Tracing"), STATGROUP_D3D12RayTracing, STATCAT_Advanced);
@@ -2832,7 +2842,7 @@ FD3D12RayTracingGeometry::FD3D12RayTracingGeometry(FRHICommandListBase& RHICmdLi
 	}
 
 	// Allocate acceleration structure buffer
-	for (uint32 GPUIndex = 0; GPUIndex < MAX_NUM_GPUS && GPUIndex < GNumExplicitGPUsForRendering; ++GPUIndex)
+	FOREACH_GPU(GPUIndex < MAX_NUM_GPUS && GPUIndex < GNumExplicitGPUsForRendering,
 	{
 		AccelerationStructureBuffers[GPUIndex] = CreateRayTracingBuffer(Adapter, GPUIndex, SizeInfo.ResultSize, ERayTracingBufferType::AccelerationStructure, DebugName);
 		AccelerationStructureBuffers[GPUIndex]->SetOwnerName(OwnerName);
@@ -2847,7 +2857,7 @@ FD3D12RayTracingGeometry::FD3D12RayTracingGeometry(FRHICommandListBase& RHICmdLi
 		{
 			INC_MEMORY_STAT_BY(STAT_D3D12RayTracingStaticBLASMemory, AccelerationStructureBuffers[GPUIndex]->GetSize());
 		}
-	}
+	});
 
 	INC_DWORD_STAT_BY(STAT_D3D12RayTracingTrianglesBLAS, Initializer.TotalPrimitiveCount);
 
@@ -2865,7 +2875,7 @@ FD3D12RayTracingGeometry::FD3D12RayTracingGeometry(FRHICommandListBase& RHICmdLi
 
 		RHICmdList.EnqueueLambda([this, SrcResourceLoc = MoveTemp(SrcResourceLoc), bForRendering](FRHICommandListBase& ExecutingCmdList)
 		{
-			for (uint32 GPUIndex = 0; GPUIndex < MAX_NUM_GPUS && GPUIndex < GNumExplicitGPUsForRendering; ++GPUIndex)
+			FOREACH_GPU(GPUIndex < MAX_NUM_GPUS && GPUIndex < GNumExplicitGPUsForRendering,
 			{
 				FD3D12CommandContext& Context = FD3D12CommandContext::Get(ExecutingCmdList, GPUIndex);
 
@@ -2885,7 +2895,7 @@ FD3D12RayTracingGeometry::FD3D12RayTracingGeometry(FRHICommandListBase& RHICmdLi
 					RegisterAsRenameListener(GPUIndex);
 					SetupHitGroupSystemParameters(GPUIndex);
 				}
-			}
+			});
 
 			SetDirty(FRHIGPUMask::All(), false);
 		});
@@ -2895,20 +2905,20 @@ FD3D12RayTracingGeometry::FD3D12RayTracingGeometry(FRHICommandListBase& RHICmdLi
 	else
 	{
 		// Offline data already registered via FD3D12RHICommandInitializeRayTracingGeometry
-		for (uint32 GPUIndex = 0; GPUIndex < MAX_NUM_GPUS && GPUIndex < GNumExplicitGPUsForRendering; ++GPUIndex)
+		FOREACH_GPU(GPUIndex < MAX_NUM_GPUS && GPUIndex < GNumExplicitGPUsForRendering,
 		{
 			RegisterAsRenameListener(GPUIndex);
-		}
+		});
 	}
 }
 
 void FD3D12RayTracingGeometry::Swap(FD3D12RayTracingGeometry& Other)
 {
-	for (uint32 i = 0; i < MAX_NUM_GPUS; i++)
+	FOREACH_GPU(GPUIndex < MAX_NUM_GPUS,
 	{
-		::Swap(AccelerationStructureBuffers[i], Other.AccelerationStructureBuffers[i]);
-		::Swap(bIsAccelerationStructureDirty[i], Other.bIsAccelerationStructureDirty[i]);
-	}
+		::Swap(AccelerationStructureBuffers[GPUIndex], Other.AccelerationStructureBuffers[GPUIndex]);
+		::Swap(bIsAccelerationStructureDirty[GPUIndex], Other.bIsAccelerationStructureDirty[GPUIndex]);
+	});
 	::Swap(AccelerationStructureCompactedSize, Other.AccelerationStructureCompactedSize);
 
 	// The rest of the members should be updated using SetInitializer()
@@ -2919,7 +2929,7 @@ void FD3D12RayTracingGeometry::ReleaseUnderlyingResource()
 	UnregisterD3D12RayTracingGeometry(this);
 
 	// Remove compaction request if still pending
-	for (uint32 GPUIndex = 0; GPUIndex < MAX_NUM_GPUS; ++GPUIndex)
+	FOREACH_GPU(GPUIndex < MAX_NUM_GPUS,
 	{
 		if (bHasPendingCompactionRequests[GPUIndex])
 		{
@@ -2929,13 +2939,13 @@ void FD3D12RayTracingGeometry::ReleaseUnderlyingResource()
 			check(bRequestFound);
 			bHasPendingCompactionRequests[GPUIndex] = false;
 		}
-	}
+	});
 
 	// Unregister as dependent resource on vertex and index buffers
-	for (uint32 GPUIndex = 0; GPUIndex < MAX_NUM_GPUS; ++GPUIndex)
+	FOREACH_GPU(GPUIndex < MAX_NUM_GPUS,
 	{
 		UnregisterAsRenameListener(GPUIndex);
-	}
+	});
 
 	for (TRefCountPtr<FD3D12Buffer>& Buffer : AccelerationStructureBuffers)
 	{
@@ -3129,10 +3139,10 @@ void FD3D12RayTracingGeometry::SetInitializer(FRHICommandListBase& RHICmdList, c
 {
 	Initializer = InInitializer;
 
-	for (uint32 GPUIndex = 0; GPUIndex < MAX_NUM_GPUS && GPUIndex < GNumExplicitGPUsForRendering; ++GPUIndex)
+	FOREACH_GPU(GPUIndex < MAX_NUM_GPUS && GPUIndex < GNumExplicitGPUsForRendering,
 	{
 		UnregisterAsRenameListener(GPUIndex);
-	}	
+	});
 
 	DebugName = !Initializer.DebugName.IsNone() ? Initializer.DebugName : FName(TEXT("BLAS"));
 
@@ -3143,11 +3153,11 @@ void FD3D12RayTracingGeometry::SetInitializer(FRHICommandListBase& RHICmdList, c
 	
 	RHICmdList.EnqueueLambda([this](FRHICommandListBase&)
 	{
-		for (uint32 GPUIndex = 0; GPUIndex < MAX_NUM_GPUS && GPUIndex < GNumExplicitGPUsForRendering; ++GPUIndex)
+		FOREACH_GPU(GPUIndex < MAX_NUM_GPUS && GPUIndex < GNumExplicitGPUsForRendering,
 		{				
 			RegisterAsRenameListener(GPUIndex);
 			SetupHitGroupSystemParameters(GPUIndex);		
-		}
+		});
 	});
 	RHICmdList.RHIThreadFence(true);
 }
diff --git a/Engine/Source/Runtime/Engine/Private/Particles/ParticlePerfStats.cpp b/Engine/Source/Runtime/Engine/Private/Particles/ParticlePerfStats.cpp
--- a/Engine/Source/Runtime/Engine/Private/Particles/ParticlePerfStats.cpp
+++ b/Engine/Source/Runtime/Engine/Private/Particles/ParticlePerfStats.cpp
@@ -1233,7 +1233,7 @@ bool FParticlePerfStatsListener_CSVProfiler::Tick()
 			for (TObjectIterator<UFXSystemComponent> It; It; ++It)
 			{
 				UFXSystemComponent* FXComponent = *It;
-				UFXSystemAsset* FXAsset = FXComponent ? FXComponent->GetFXSystemAsset() : nullptr;
+				UFXSystemAsset* FXAsset = FXComponent->GetFXSystemAsset();
 				if (!IsValidChecked(FXComponent) || !IsValid(FXAsset) || FXComponent->IsUnreachable() || FXComponent->HasAnyFlags(EObjectFlags::RF_ClassDefaultObject))
 				{
 					continue;
diff --git a/Engine/Source/Runtime/Engine/Private/UnrealEngine.cpp b/Engine/Source/Runtime/Engine/Private/UnrealEngine.cpp
--- a/Engine/Source/Runtime/Engine/Private/UnrealEngine.cpp
+++ b/Engine/Source/Runtime/Engine/Private/UnrealEngine.cpp
@@ -6844,7 +6844,7 @@ bool UEngine::HandleListSkeletalMeshesCommand(const TCHAR* Cmd, FOutputDevice& A
 		USkeletalMesh* Mesh = *It;
 
 		bool bUnknownRef = false;
-		bool bIsStreaming = (Mesh != nullptr ? Mesh->GetStreamingIndex() != INDEX_NONE : false);
+		bool bIsStreaming = Mesh->GetStreamingIndex() != INDEX_NONE;
 		
 		if (Streamer)
 		{
