Remove multi-layer support from FRHIRayTracingScene.
- high level code should create a FRHIRayTracingScene per layer
- multiple FRHIRayTracingScene can share the same FRHIShaderBindingTable to match 5.4 behaviour.

diff --git a/Engine/Source/Runtime/Apple/MetalRHI/Private/MetalRayTracing.cpp b/Engine/Source/Runtime/Apple/MetalRHI/Private/MetalRayTracing.cpp
--- a/Engine/Source/Runtime/Apple/MetalRHI/Private/MetalRayTracing.cpp
+++ b/Engine/Source/Runtime/Apple/MetalRHI/Private/MetalRayTracing.cpp
@@ -433,32 +433,10 @@ void FMetalRayTracingGeometry::RebuildDescriptors()
 FMetalRayTracingScene::FMetalRayTracingScene(FRayTracingSceneInitializer2 InInitializer)
 	: Initializer(MoveTemp(InInitializer))
 {
-	SizeInfo = {};
+	MTL::InstanceAccelerationStructureDescriptor* InstanceDescriptor;
+	InstanceDescriptor.SetInstanceCount(Initializer.NumNativeInstances);
 
-	PRAGMA_DISABLE_DEPRECATION_WARNINGS
-	const uint32 NumLayers = Initializer.NumNativeInstancesPerLayer.Num();
-	PRAGMA_ENABLE_DEPRECATION_WARNINGS
-	check(NumLayers > 0);
-
-	Layers.SetNum(NumLayers);
-
-	for (uint32 LayerIndex = 0; LayerIndex < NumLayers; ++LayerIndex)
-	{
-		FLayerData& Layer = Layers[LayerIndex];
-
-		MTL::InstanceAccelerationStructureDescriptor* InstanceDescriptor;
-		PRAGMA_DISABLE_DEPRECATION_WARNINGS
-		InstanceDescriptor.SetInstanceCount(Initializer.NumNativeInstancesPerLayer[LayerIndex]);
-		PRAGMA_ENABLE_DEPRECATION_WARNINGS
-
-		Layer.SizeInfo = CalcRayTracingGeometrySize(InstanceDescriptor);
-
-		Layer.BufferOffset = Align(SizeInfo.ResultSize, GRHIRayTracingAccelerationStructureAlignment);
-		Layer.ScratchBufferOffset = Align(SizeInfo.BuildScratchSize, GRHIRayTracingScratchBufferAlignment);
-
-		SizeInfo.ResultSize = Layer.BufferOffset + Layer.SizeInfo.ResultSize;
-		SizeInfo.BuildScratchSize = Layer.BufferOffset + Layer.SizeInfo.BuildScratchSize;
-	}
+	SizeInfo = CalcRayTracingGeometrySize(InstanceDescriptor);
 
 	MutableAccelerationStructures = [[NSMutableArray<id<MTLAccelerationStructure>> new] init];
 
@@ -492,17 +470,13 @@ void FMetalRayTracingScene::BindBuffer(FRHIBuffer* InBuffer, uint32 InBufferOffs
 	FMetalDeviceContext& Context = GetMetalDeviceContext();
 	MTL::Device* Device = Context.GetDevice();
 
-	for (auto& Layer : Layers)
 	{
-		checkf(Layer.ShaderResourceView == nullptr, TEXT("Binding multiple buffers is not currently supported."));
-
-		const uint32 LayerOffset = InBufferOffset + Layer.BufferOffset;
-		check(LayerOffset % GRHIRayTracingAccelerationStructureAlignment == 0);
-
+		checkf(ShaderResourceView == nullptr, TEXT("Binding multiple buffers is not currently supported."));
+		check(InBufferOffset % GRHIRayTracingAccelerationStructureAlignment == 0);
 		check(AccelerationStructureBuffer->IsAccelerationStructure());
 
-		FShaderResourceViewInitializer ViewInitializer(AccelerationStructureBuffer, LayerOffset, 0);
-		Layer.ShaderResourceView = new FMetalShaderResourceView(ViewInitializer);
+		FShaderResourceViewInitializer ViewInitializer(AccelerationStructureBuffer, InBufferOffset, 0);
+		ShaderResourceView = new FMetalShaderResourceView(ViewInitializer);
 
 		FString DebugNameString = Initializer.DebugName.ToString();
 		DebugNameString = (DebugNameString.IsEmpty()) ? TEXT("TLAS") : DebugNameString;
@@ -624,29 +598,19 @@ void FMetalRayTracingScene::BuildAccelerationStructure(
 	MTL::AccelerationStructureCommandEncoder* CommandEncoder = Encoder.GetAccelerationStructureCommandEncoder();
 	check(CommandEncoder.GetPtr());
 
-	PRAGMA_DISABLE_DEPRECATION_WARNINGS
-	const uint32 NumLayers = Initializer.NumNativeInstancesPerLayer.Num();
-	PRAGMA_ENABLE_DEPRECATION_WARNINGS
-
 	FMetalBuffer CurScratchBuffer = InScratchBuffer->GetCurrentBuffer();
 	check(CurScratchBuffer);
 
-	check(Layers.Num() && NumLayers > 0);
-	for (uint32 LayerIndex = 0; LayerIndex < NumLayers; ++LayerIndex)
 	{
-		FLayerData& Layer = Layers[LayerIndex];
-
 		MTL::InstanceAccelerationStructureDescriptor* InstanceDescriptor = MTL::InstanceAccelerationStructureDescriptor();
-		PRAGMA_DISABLE_DEPRECATION_WARNINGS
-		InstanceDescriptor.SetInstanceCount(Initializer.NumNativeInstancesPerLayer[LayerIndex]);
-		PRAGMA_ENABLE_DEPRECATION_WARNINGS
+		InstanceDescriptor.SetInstanceCount(Initializer.NumNativeInstances);
 		InstanceDescriptor.SetInstanceDescriptorBuffer(CurInstanceBuffer);
 		InstanceDescriptor.SetInstanceDescriptorBufferOffset(InstanceBufferOffset);
 		InstanceDescriptor.SetInstancedAccelerationStructures((__bridge NSArray*)MutableAccelerationStructures);
 		InstanceDescriptor.SetInstanceDescriptorStride(GRHIRayTracingInstanceDescriptorSize);
 		InstanceDescriptor.SetInstanceDescriptorType(MTL::AccelerationStructureInstanceDescriptorType::UserID);
 
-		MTL::AccelerationStructure* AS = ResourceCast(Layer.ShaderResourceView->GetBuffer())->AccelerationStructureHandle;
+		MTL::AccelerationStructure* AS = ResourceCast(ShaderResourceView->GetBuffer())->AccelerationStructureHandle;
 		CommandEncoder.BuildAccelerationStructure(AS, InstanceDescriptor, CurScratchBuffer, ScratchOffset);
 	}
 
diff --git a/Engine/Source/Runtime/Apple/MetalRHI/Private/MetalRayTracing.h b/Engine/Source/Runtime/Apple/MetalRHI/Private/MetalRayTracing.h
--- a/Engine/Source/Runtime/Apple/MetalRHI/Private/MetalRayTracing.h
+++ b/Engine/Source/Runtime/Apple/MetalRHI/Private/MetalRayTracing.h
@@ -88,18 +88,10 @@ public:
 	void BuildPerInstanceGeometryParameterBuffer();
 
 	inline const FRayTracingSceneInitializer2& GetInitializer() const override final { return Initializer; }
-	inline uint32 GetLayerBufferOffset(uint32 LayerIndex) const override final { return Layers[LayerIndex].BufferOffset; }
 
 	TRefCountPtr<FMetalShaderResourceView> InstanceBufferSRV;
 
-	struct FLayerData
-	{
-		TRefCountPtr<FMetalShaderResourceView> ShaderResourceView;
-		uint32 BufferOffset;
-		uint32 ScratchBufferOffset;
-		FRayTracingAccelerationStructureSize SizeInfo;
-	};
-	TArray<FLayerData> Layers;
+	TRefCountPtr<FMetalShaderResourceView> ShaderResourceView;
 
 private:
 	friend class FMetalRHICommandContext;
@@ -108,7 +100,7 @@ private:
 	/** The initializer provided to build the scene. Contains all the free standing stuff that used to be owned by the RT implementation. */
 	const FRayTracingSceneInitializer2 Initializer;
 
-	/** Acceleration Structure for the whole scene (shared between each layer). */
+	/** Acceleration Structure for the whole scene. */
 	TRefCountPtr<FMetalRHIBuffer> AccelerationStructureBuffer;
 
 	/** Root Constants for geometry evaluation in HitGroup/Miss (emulates D3D12 RootConstants with a global scope). */
diff --git a/Engine/Source/Runtime/D3D12RHI/Private/D3D12RayTracing.cpp b/Engine/Source/Runtime/D3D12RHI/Private/D3D12RayTracing.cpp
--- a/Engine/Source/Runtime/D3D12RHI/Private/D3D12RayTracing.cpp
+++ b/Engine/Source/Runtime/D3D12RHI/Private/D3D12RayTracing.cpp
@@ -3396,38 +3396,14 @@ FD3D12RayTracingScene::FD3D12RayTracingScene(FD3D12Adapter* Adapter, FRayTracing
 	PRAGMA_ENABLE_DEPRECATION_WARNINGS
 	checkf(Initializer.Lifetime == RTSL_SingleFrame, TEXT("Only single-frame ray tracing scenes are currently implemented."));
 
-	SizeInfo = {};
-
-	PRAGMA_DISABLE_DEPRECATION_WARNINGS
-	const uint32 NumLayers = Initializer.NumNativeInstancesPerLayer.Num();
-	PRAGMA_ENABLE_DEPRECATION_WARNINGS
-	check(NumLayers > 0);
-
-	Layers.SetNum(NumLayers);
-
-	for (uint32 LayerIndex = 0; LayerIndex < NumLayers; ++LayerIndex)
-	{
-		FLayerData& Layer = Layers[LayerIndex];
-
-		Layer.BuildInputs = {};
-		Layer.BuildInputs.Type = D3D12_RAYTRACING_ACCELERATION_STRUCTURE_TYPE_TOP_LEVEL;
-		Layer.BuildInputs.DescsLayout = D3D12_ELEMENTS_LAYOUT_ARRAY;
-		PRAGMA_DISABLE_DEPRECATION_WARNINGS
-		Layer.BuildInputs.NumDescs = Initializer.NumNativeInstancesPerLayer[LayerIndex];
-		PRAGMA_ENABLE_DEPRECATION_WARNINGS
-		Layer.BuildInputs.Flags = TranslateRayTracingAccelerationStructureFlags(Initializer.BuildFlags);
-
-		// Get maximum buffer sizes for all GPUs in the system
-		Layer.SizeInfo = RHICalcRayTracingSceneSize(Layer.BuildInputs.NumDescs, Initializer.BuildFlags);
-
-		Layer.BufferOffset = Align(SizeInfo.ResultSize, GRHIRayTracingAccelerationStructureAlignment);
-		Layer.BuildScratchOffset = Align(SizeInfo.BuildScratchSize, GRHIRayTracingScratchBufferAlignment);
-		Layer.UpdateScratchOffset = Align(SizeInfo.UpdateScratchSize, GRHIRayTracingScratchBufferAlignment);
+	BuildInputs = {};
+	BuildInputs.Type = D3D12_RAYTRACING_ACCELERATION_STRUCTURE_TYPE_TOP_LEVEL;
+	BuildInputs.DescsLayout = D3D12_ELEMENTS_LAYOUT_ARRAY;
+	BuildInputs.NumDescs = Initializer.NumNativeInstances;
+	BuildInputs.Flags = TranslateRayTracingAccelerationStructureFlags(Initializer.BuildFlags);
 
-		SizeInfo.ResultSize = Layer.BufferOffset + Layer.SizeInfo.ResultSize;
-		SizeInfo.BuildScratchSize = Layer.BuildScratchOffset + Layer.SizeInfo.BuildScratchSize;
-		SizeInfo.UpdateScratchSize = Layer.UpdateScratchOffset + Layer.SizeInfo.UpdateScratchSize;
-	}
+	// Get maximum buffer sizes for all GPUs in the system
+	SizeInfo = RHICalcRayTracingSceneSize(BuildInputs.NumDescs, Initializer.BuildFlags);
 };
 
 FD3D12RayTracingScene::~FD3D12RayTracingScene()
@@ -3468,13 +3444,8 @@ void FD3D12RayTracingScene::BindBuffer(FRHIBuffer* InBuffer, uint32 InBufferOffs
 		INC_MEMORY_STAT_BY(STAT_D3D12RayTracingUsedVideoMemory, AccelerationStructureBuffers[GPUIndex]->GetSize());
 		INC_MEMORY_STAT_BY(STAT_D3D12RayTracingTLASMemory, AccelerationStructureBuffers[GPUIndex]->GetSize());
 	}
-	BufferOffset = InBufferOffset;
 
-	for (auto& Layer : Layers)
-	{
-		const uint32 LayerOffset = InBufferOffset + Layer.BufferOffset;
-		check(LayerOffset % GRHIRayTracingAccelerationStructureAlignment == 0);
-	}
+	BufferOffset = InBufferOffset;
 }
 
 void FD3D12RayTracingScene::BuildAccelerationStructure(FD3D12CommandContext& CommandContext,
@@ -3505,51 +3476,40 @@ void FD3D12RayTracingScene::BuildAccelerationStructure(FD3D12CommandContext& Com
 
 	checkf(ScratchBuffer, TEXT("TLAS build requires scratch buffer of at least %lld bytes."), SizeInfo.BuildScratchSize);
 
-	PRAGMA_DISABLE_DEPRECATION_WARNINGS
-	const uint32 NumLayers = Initializer.NumNativeInstancesPerLayer.Num();
-	PRAGMA_ENABLE_DEPRECATION_WARNINGS
-
-	for (uint32 LayerIndex = 0; LayerIndex < NumLayers; ++LayerIndex)
 	{
-		const FLayerData& Layer = Layers[LayerIndex];
-
 		D3D12_RAYTRACING_ACCELERATION_STRUCTURE_PREBUILD_INFO PrebuildInfo = {};
 
-		CommandContext.GetParentDevice()->GetRaytracingAccelerationStructurePrebuildInfo(&Layer.BuildInputs, &PrebuildInfo);
+		CommandContext.GetParentDevice()->GetRaytracingAccelerationStructurePrebuildInfo(&BuildInputs, &PrebuildInfo);
 
-		checkf(PrebuildInfo.ResultDataMaxSizeInBytes <= Layer.SizeInfo.ResultSize,
+		checkf(PrebuildInfo.ResultDataMaxSizeInBytes <= SizeInfo.ResultSize,
 			TEXT("TLAS build result buffer now requires %lld bytes, but only %lld was calculated in the constructor."),
-			PrebuildInfo.ResultDataMaxSizeInBytes, Layer.SizeInfo.ResultSize);
+			PrebuildInfo.ResultDataMaxSizeInBytes, SizeInfo.ResultSize);
 
-		checkf(PrebuildInfo.ScratchDataSizeInBytes <= Layer.SizeInfo.BuildScratchSize,
+		checkf(PrebuildInfo.ScratchDataSizeInBytes <= SizeInfo.BuildScratchSize,
 			TEXT("TLAS build scratch buffer now requires %lld bytes, but only %lld was calculated in the constructor."),
-			PrebuildInfo.ScratchDataSizeInBytes, Layer.SizeInfo.BuildScratchSize);
+			PrebuildInfo.ScratchDataSizeInBytes, SizeInfo.BuildScratchSize);
 
-		checkf(PrebuildInfo.UpdateScratchDataSizeInBytes <= Layer.SizeInfo.UpdateScratchSize,
+		checkf(PrebuildInfo.UpdateScratchDataSizeInBytes <= SizeInfo.UpdateScratchSize,
 			TEXT("TLAS update scratch buffer now requires %lld bytes, but only %lld was calculated in the constructor."),
-			PrebuildInfo.UpdateScratchDataSizeInBytes, Layer.SizeInfo.UpdateScratchSize);
+			PrebuildInfo.UpdateScratchDataSizeInBytes, SizeInfo.UpdateScratchSize);
 
 		if (bIsUpdate)
 		{
-			PRAGMA_DISABLE_DEPRECATION_WARNINGS
-			checkf(ScratchBufferOffset + Layer.UpdateScratchOffset + PrebuildInfo.UpdateScratchDataSizeInBytes <= ScratchBuffer->GetSize(),
+			checkf(ScratchBufferOffset + PrebuildInfo.UpdateScratchDataSizeInBytes <= ScratchBuffer->GetSize(),
 				TEXT("TLAS scratch buffer size is %d bytes with offset %d (%d bytes available), but the update requires %lld bytes. ")
 				TEXT("BuildInputs.NumDescs = %d, Instances.Num = %d."),
-				ScratchBuffer->GetSize(), ScratchBufferOffset + Layer.UpdateScratchOffset, ScratchBuffer->GetSize() - ScratchBufferOffset,
+				ScratchBuffer->GetSize(), ScratchBufferOffset, ScratchBuffer->GetSize() - ScratchBufferOffset,
 				PrebuildInfo.UpdateScratchDataSizeInBytes,
-				Layer.BuildInputs.NumDescs, Initializer.NumNativeInstancesPerLayer[LayerIndex]);
-			PRAGMA_ENABLE_DEPRECATION_WARNINGS
+				BuildInputs.NumDescs, Initializer.NumNativeInstances);
 		}
 		else
 		{
-			PRAGMA_DISABLE_DEPRECATION_WARNINGS
-			checkf(ScratchBufferOffset + Layer.BuildScratchOffset + PrebuildInfo.ScratchDataSizeInBytes <= ScratchBuffer->GetSize(),
+			checkf(ScratchBufferOffset + PrebuildInfo.ScratchDataSizeInBytes <= ScratchBuffer->GetSize(),
 				TEXT("TLAS scratch buffer size is %d bytes with offset %d (%d bytes available), but the build requires %lld bytes. ")
 				TEXT("BuildInputs.NumDescs = %d, Instances.Num = %d."),
-				ScratchBuffer->GetSize(), ScratchBufferOffset + Layer.BuildScratchOffset, ScratchBuffer->GetSize() - ScratchBufferOffset,
+				ScratchBuffer->GetSize(), ScratchBufferOffset, ScratchBuffer->GetSize() - ScratchBufferOffset,
 				PrebuildInfo.ScratchDataSizeInBytes,
-				Layer.BuildInputs.NumDescs, Initializer.NumNativeInstancesPerLayer[LayerIndex]);
-			PRAGMA_ENABLE_DEPRECATION_WARNINGS
+				BuildInputs.NumDescs, Initializer.NumNativeInstances);
 		}
 	}
 
@@ -3657,19 +3617,15 @@ void FD3D12RayTracingScene::BuildAccelerationStructure(FD3D12CommandContext& Com
 
 	if (ShouldRunRayTracingGPUValidation())
 	{
-		for (uint32 LayerIndex = 0; LayerIndex < NumLayers; ++LayerIndex)
-		{
-			const FLayerData& Layer = Layers[LayerIndex];
-			TRHICommandList_RecursiveHazardous<FD3D12CommandContext> RHICmdList(&CommandContext);
-			uint32 InstanceBufferStride = GRHIRayTracingInstanceDescriptorSize;
-			PRAGMA_DISABLE_DEPRECATION_WARNINGS
-			// TODO: Validation related to SBT needs to be done somewhere else since SBT is not known when in BuildAccelerationStructure
-			uint32 TotalHitGroupSlots = Initializer.NumTotalSegments * Initializer.ShaderSlotsPerGeometrySegment;
-			PRAGMA_ENABLE_DEPRECATION_WARNINGS
-			FRayTracingValidateSceneBuildParamsCS::Dispatch(RHICmdList,
-				TotalHitGroupSlots, Layer.BuildInputs.NumDescs,
-				InstanceBuffer, InstanceBufferOffset, InstanceBufferStride);
-		}
+		TRHICommandList_RecursiveHazardous<FD3D12CommandContext> RHICmdList(&CommandContext);
+		uint32 InstanceBufferStride = GRHIRayTracingInstanceDescriptorSize;
+		PRAGMA_DISABLE_DEPRECATION_WARNINGS
+		// TODO: Validation related to SBT needs to be done somewhere else since SBT is not known when in BuildAccelerationStructure
+		uint32 TotalHitGroupSlots = Initializer.NumTotalSegments * Initializer.ShaderSlotsPerGeometrySegment;
+		PRAGMA_ENABLE_DEPRECATION_WARNINGS
+		FRayTracingValidateSceneBuildParamsCS::Dispatch(RHICmdList,
+			TotalHitGroupSlots, BuildInputs.NumDescs,
+			InstanceBuffer, InstanceBufferOffset, InstanceBufferStride);
 	}
 
 	// UAV barrier is used here to ensure that all bottom level acceleration structures are built
@@ -3677,28 +3633,14 @@ void FD3D12RayTracingScene::BuildAccelerationStructure(FD3D12CommandContext& Com
 	CommandContext.FlushResourceBarriers();
 
 	TArray<D3D12_BUILD_RAYTRACING_ACCELERATION_STRUCTURE_DESC, TInlineAllocator<32>> BuildDescs;
-	BuildDescs.Reserve(NumLayers);
-
-	uint32 InstanceBaseOffset = 0;
+	BuildDescs.Reserve(1);
 
-	for (uint32 LayerIndex = 0; LayerIndex < NumLayers; ++LayerIndex)
 	{
-		const FLayerData& Layer = Layers[LayerIndex];
-
-		const D3D12_GPU_VIRTUAL_ADDRESS BufferAddress = AccelerationStructureBuffer->ResourceLocation.GetGPUVirtualAddress() + BufferOffset + Layer.BufferOffset;
+		const D3D12_GPU_VIRTUAL_ADDRESS BufferAddress = AccelerationStructureBuffer->ResourceLocation.GetGPUVirtualAddress() + BufferOffset;
 		D3D12_GPU_VIRTUAL_ADDRESS ScratchAddress = ScratchBuffer->ResourceLocation.GetGPUVirtualAddress() + ScratchBufferOffset;
 
-		if (bIsUpdate)
-		{
-			ScratchAddress += Layer.UpdateScratchOffset;
-		}
-		else
-		{
-			ScratchAddress += Layer.BuildScratchOffset;
-		}
-
 		checkf(BufferAddress % GRHIRayTracingAccelerationStructureAlignment == 0,
-			TEXT("TLAS  buffer (plus offset) must be aligned to %lld bytes."),
+			TEXT("TLAS buffer (plus offset) must be aligned to %lld bytes."),
 			GRHIRayTracingAccelerationStructureAlignment);
 
 		checkf(ScratchAddress % GRHIRayTracingScratchBufferAlignment == 0,
@@ -3706,8 +3648,8 @@ void FD3D12RayTracingScene::BuildAccelerationStructure(FD3D12CommandContext& Com
 			GRHIRayTracingScratchBufferAlignment);
 
 		D3D12_BUILD_RAYTRACING_ACCELERATION_STRUCTURE_DESC& BuildDesc = BuildDescs.AddDefaulted_GetRef();
-		BuildDesc.Inputs = Layer.BuildInputs;
-		BuildDesc.Inputs.InstanceDescs = InstanceBuffer->ResourceLocation.GetGPUVirtualAddress() + InstanceBufferOffset + InstanceBaseOffset * GRHIRayTracingInstanceDescriptorSize;
+		BuildDesc.Inputs = BuildInputs;
+		BuildDesc.Inputs.InstanceDescs = InstanceBuffer->ResourceLocation.GetGPUVirtualAddress() + InstanceBufferOffset;
 
 		if (bIsUpdate)
 		{
@@ -3729,10 +3671,6 @@ void FD3D12RayTracingScene::BuildAccelerationStructure(FD3D12CommandContext& Com
 		{
 			INC_DWORD_STAT(STAT_D3D12RayTracingBuiltTLAS);
 		}
-
-		PRAGMA_DISABLE_DEPRECATION_WARNINGS
-		InstanceBaseOffset += Initializer.NumNativeInstancesPerLayer[LayerIndex];
-		PRAGMA_ENABLE_DEPRECATION_WARNINGS
 	}
 
 	CommandContext.BuildAccelerationStructuresInternal(BuildDescs);
diff --git a/Engine/Source/Runtime/D3D12RHI/Private/D3D12RayTracing.h b/Engine/Source/Runtime/D3D12RHI/Private/D3D12RayTracing.h
--- a/Engine/Source/Runtime/D3D12RHI/Private/D3D12RayTracing.h
+++ b/Engine/Source/Runtime/D3D12RHI/Private/D3D12RayTracing.h
@@ -99,7 +99,6 @@ public:
 	~FD3D12RayTracingScene();
 
 	const FRayTracingSceneInitializer2& GetInitializer() const override final { return Initializer; }
-	uint32 GetLayerBufferOffset(uint32 LayerIndex) const override final { return Layers[LayerIndex].BufferOffset; }
 
 	void BindBuffer(FRHIBuffer* Buffer, uint32 BufferOffset);
 	void ReleaseBuffer();
@@ -110,16 +109,7 @@ public:
 		EAccelerationStructureBuildMode BuildMode
 	);
 
-	struct FLayerData
-	{
-		D3D12_BUILD_RAYTRACING_ACCELERATION_STRUCTURE_INPUTS BuildInputs = {};
-		FRayTracingAccelerationStructureSize SizeInfo = {};
-		uint32 BufferOffset;
-		uint32 BuildScratchOffset;
-		uint32 UpdateScratchOffset;
-	};
-
-	TArray<FLayerData> Layers;
+	D3D12_BUILD_RAYTRACING_ACCELERATION_STRUCTURE_INPUTS BuildInputs = {};
 
 	TRefCountPtr<FD3D12Buffer> AccelerationStructureBuffers[MAX_NUM_GPUS];
 	uint32 BufferOffset = 0;
diff --git a/Engine/Source/Runtime/D3D12RHI/Private/D3D12RayTracingDebug.cpp b/Engine/Source/Runtime/D3D12RHI/Private/D3D12RayTracingDebug.cpp
--- a/Engine/Source/Runtime/D3D12RHI/Private/D3D12RayTracingDebug.cpp
+++ b/Engine/Source/Runtime/D3D12RHI/Private/D3D12RayTracingDebug.cpp
@@ -38,7 +38,6 @@ static void DebugSerializeScene(const FD3D12RayTracingScene& Scene, FD3D12Buffer
 	// All buffers are vector-aligned to allow efficient loading into GPU memory and subsequent GPU access.
 	// Serialized scene data is expected to be loaded all at once from disk and copied into a single GPU buffer.
 	static constexpr uint32 AlignmentRequirement = 16;
-	static constexpr uint32 MaxNumLayers = 64;
 
 	// Serialized types
 
@@ -58,13 +57,10 @@ static void DebugSerializeScene(const FD3D12RayTracingScene& Scene, FD3D12Buffer
 			uint32 Strings = 0;
 		} Offsets;
 
-		uint32 NumLayers = 0;
-		uint32 PerLayerNumInstances[MaxNumLayers] = {};
 		uint32 NumInstances = 0;
 		uint32 NumGeometries = 0;
 		uint32 NumBuffers = 0;
 		uint32 NumStrings = 0;
-		uint32 Padding[2] = {};
 	};
 	static_assert(sizeof(FSceneHeader) % AlignmentRequirement == 0, "Serialized scene data must be vector-aligned");
 
@@ -159,17 +155,7 @@ static void DebugSerializeScene(const FD3D12RayTracingScene& Scene, FD3D12Buffer
 	}
 
 	{
-		// Per-layer number of instances
-
-		SceneHeader.NumLayers = FMath::Min<uint32>(MaxNumLayers, Scene.Layers.Num());
-
-		for (uint32 LayerIndex = 0; LayerIndex < SceneHeader.NumLayers; ++LayerIndex)
-		{
-			const FD3D12RayTracingScene::FLayerData& Layer = Scene.Layers[LayerIndex];
-			uint32 NumInstances = Layer.BuildInputs.NumDescs;
-			SceneHeader.PerLayerNumInstances[LayerIndex] = NumInstances;
-			SceneHeader.NumInstances += NumInstances;
-		}
+		SceneHeader.NumInstances = Scene.BuildInputs.NumDescs;
 
 		// Instance buffer
 
diff --git a/Engine/Source/Runtime/RHI/Public/DynamicRHI.h b/Engine/Source/Runtime/RHI/Public/DynamicRHI.h
--- a/Engine/Source/Runtime/RHI/Public/DynamicRHI.h
+++ b/Engine/Source/Runtime/RHI/Public/DynamicRHI.h
@@ -1456,10 +1456,9 @@ FORCEINLINE FRayTracingAccelerationStructureSize RHICalcRayTracingSceneSize(uint
 FORCEINLINE FRayTracingSceneRHIRef RHICreateRayTracingScene(FRayTracingSceneInitializer2 Initializer)
 {
 	PRAGMA_DISABLE_DEPRECATION_WARNINGS
-	if(Initializer.NumNativeInstances > 0 || Initializer.NumNativeInstancesPerLayer.IsEmpty())
+	if (Initializer.NumNativeInstancesPerLayer.Num() >= 1)
 	{
-		checkf(Initializer.NumNativeInstancesPerLayer.IsEmpty(), TEXT("Can't specify NumNativeInstances and NumNativeInstancesPerLayer simultaneously."));
-		Initializer.NumNativeInstancesPerLayer.Add(Initializer.NumNativeInstances);
+		Initializer.NumNativeInstances = Initializer.NumNativeInstancesPerLayer[0];
 	}
 	PRAGMA_ENABLE_DEPRECATION_WARNINGS
 	return GDynamicRHI->RHICreateRayTracingScene(MoveTemp(Initializer));
diff --git a/Engine/Source/Runtime/RHI/Public/RHIResources.h b/Engine/Source/Runtime/RHI/Public/RHIResources.h
--- a/Engine/Source/Runtime/RHI/Public/RHIResources.h
+++ b/Engine/Source/Runtime/RHI/Public/RHIResources.h
@@ -3407,7 +3407,11 @@ public:
 	}
 
 	UE_DEPRECATED(5.5, "FRHIRayTracingScene layers are deprecated. Create one FRHIRayTracingScene per layer instead.")
-	virtual uint32 GetLayerBufferOffset(uint32 LayerIndex) const = 0;
+	virtual uint32 GetLayerBufferOffset(uint32 LayerIndex) const
+	{
+		checkf(LayerIndex == 0, TEXT("FRHIRayTracingScene layers are deprecated. Create one FRHIRayTracingScene per layer instead."));
+		return 0;
+	}
 
 	UE_DEPRECATED(5.5, "Create standalone FRHIShaderBindingTable instead.")
 	virtual FRHIShaderBindingTable* FindOrCreateShaderBindingTable(const FRHIRayTracingPipelineState* Pipeline) = 0;
diff --git a/Engine/Source/Runtime/Renderer/Private/RayTracing/RayTracingInstanceBufferUtil.cpp b/Engine/Source/Runtime/Renderer/Private/RayTracing/RayTracingInstanceBufferUtil.cpp
--- a/Engine/Source/Runtime/Renderer/Private/RayTracing/RayTracingInstanceBufferUtil.cpp
+++ b/Engine/Source/Runtime/Renderer/Private/RayTracing/RayTracingInstanceBufferUtil.cpp
@@ -163,19 +163,6 @@ void FillRayTracingInstanceUploadBuffer(
 
 	const FRayTracingSceneInitializer2& SceneInitializer = RayTracingSceneRHI->GetInitializer();
 
-	PRAGMA_DISABLE_DEPRECATION_WARNINGS
-	const uint32 NumLayers = SceneInitializer.NumNativeInstancesPerLayer.Num();
-
-	TArray<uint32> LayerBaseIndices;
-	LayerBaseIndices.SetNumUninitialized(NumLayers);
-	LayerBaseIndices[0] = 0;
-
-	for (uint32 LayerIndex = 1; LayerIndex < NumLayers; ++LayerIndex)
-	{
-		LayerBaseIndices[LayerIndex] = LayerBaseIndices[LayerIndex - 1] + SceneInitializer.NumNativeInstancesPerLayer[LayerIndex - 1];
-	}
-	PRAGMA_ENABLE_DEPRECATION_WARNINGS
-
 	const int32 NumSceneInstances = Instances.Num();
 	const int32 MinBatchSize = 128;
 	ParallelFor(
@@ -191,7 +178,6 @@ void FillRayTracingInstanceUploadBuffer(
 			InstanceGeometryIndices,
 			BaseUploadBufferOffsets,
 			BaseInstancePrefixSum,
-			LayerBaseIndices,
 			PreViewTranslation,
 			&SceneInitializer
 		](int32 SceneInstanceIndex)
@@ -211,9 +197,6 @@ void FillRayTracingInstanceUploadBuffer(
 			checkf(bGpuSceneInstance + bCpuInstance == 1, TEXT("Instance can only get transforms from one of GPUScene, or Transforms array."));
 
 			const uint32 AccelerationStructureIndex = InstanceGeometryIndices[SceneInstanceIndex];
-			PRAGMA_DISABLE_DEPRECATION_WARNINGS
-			const uint32 LayerBaseIndex = LayerBaseIndices[SceneInstance.LayerIndex];
-			PRAGMA_ENABLE_DEPRECATION_WARNINGS
 			const uint32 BaseInstanceIndex = BaseInstancePrefixSum[SceneInstanceIndex];
 			const uint32 BaseTransformIndex = bCpuInstance ? BaseUploadBufferOffsets[SceneInstanceIndex] : 0;
 
@@ -261,7 +244,7 @@ void FillRayTracingInstanceUploadBuffer(
 					}
 				}
 
-				InstanceDesc.OutputDescriptorIndex = LayerBaseIndex + BaseInstanceIndex + TransformIndex;
+				InstanceDesc.OutputDescriptorIndex = BaseInstanceIndex + TransformIndex;
 				InstanceDesc.AccelerationStructureIndex = AccelerationStructureIndex;
 				InstanceDesc.InstanceId = UserData;
 				InstanceDesc.InstanceMaskAndFlags = SceneInstance.Mask | ((uint32)SceneInstance.Flags << 8);
diff --git a/Engine/Source/Runtime/VulkanRHI/Private/VulkanRayTracing.cpp b/Engine/Source/Runtime/VulkanRHI/Private/VulkanRayTracing.cpp
--- a/Engine/Source/Runtime/VulkanRHI/Private/VulkanRayTracing.cpp
+++ b/Engine/Source/Runtime/VulkanRHI/Private/VulkanRayTracing.cpp
@@ -676,30 +676,7 @@ FVulkanRayTracingScene::FVulkanRayTracingScene(FRayTracingSceneInitializer2 InIn
 {
 	INC_DWORD_STAT(STAT_VulkanRayTracingAllocatedTLAS);
 
-	SizeInfo = {};
-
-	PRAGMA_DISABLE_DEPRECATION_WARNINGS
-	const uint32 NumLayers = Initializer.NumNativeInstancesPerLayer.Num();
-	PRAGMA_ENABLE_DEPRECATION_WARNINGS
-	check(NumLayers > 0);
-
-	Layers.SetNum(NumLayers);
-
-	for (uint32 LayerIndex = 0; LayerIndex < NumLayers; ++LayerIndex)
-	{
-		FLayerData& Layer = Layers[LayerIndex];
-
-		PRAGMA_DISABLE_DEPRECATION_WARNINGS
-		Layer.SizeInfo = RHICalcRayTracingSceneSize(Initializer.NumNativeInstancesPerLayer[LayerIndex], Initializer.BuildFlags);
-		PRAGMA_ENABLE_DEPRECATION_WARNINGS
-		Layer.BufferOffset = Align(SizeInfo.ResultSize, GRHIRayTracingAccelerationStructureAlignment);
-		Layer.BuildScratchOffset = Align(SizeInfo.BuildScratchSize, GRHIRayTracingScratchBufferAlignment);
-		Layer.UpdateScratchOffset = Align(SizeInfo.UpdateScratchSize, GRHIRayTracingScratchBufferAlignment);
-
-		SizeInfo.ResultSize = Layer.BufferOffset + Layer.SizeInfo.ResultSize;
-		SizeInfo.BuildScratchSize = Layer.BuildScratchOffset + Layer.SizeInfo.BuildScratchSize;
-		SizeInfo.UpdateScratchSize = Layer.UpdateScratchOffset + Layer.SizeInfo.UpdateScratchSize;
-	}
+	SizeInfo = RHICalcRayTracingSceneSize(Initializer.NumNativeInstances, Initializer.BuildFlags);
 
 	const uint32 ParameterBufferSize = FMath::Max<uint32>(1, Initializer.NumTotalSegments) * sizeof(FVulkanRayTracingGeometryParameters);
 	FRHIResourceCreateInfo ParameterBufferCreateInfo(TEXT("RayTracingSceneMetadata"));
@@ -739,20 +716,18 @@ void FVulkanRayTracingScene::BindBuffer(FRHIBuffer* InBuffer, uint32 InBufferOff
 	INC_MEMORY_STAT_BY(STAT_VulkanRayTracingUsedVideoMemory, AccelerationStructureBuffer->GetSize());
 	INC_MEMORY_STAT_BY(STAT_VulkanRayTracingTLASMemory, AccelerationStructureBuffer->GetSize());
 
-	for (auto& Layer : Layers)
 	{
-		checkf(!Layer.View.IsValid(), TEXT("Binding multiple buffers is not currently supported."));
+		checkf(!View.IsValid(), TEXT("Binding multiple buffers is not currently supported."));
 
-		const uint32 LayerOffset = InBufferOffset + Layer.BufferOffset;
-		check(LayerOffset % GRHIRayTracingAccelerationStructureAlignment == 0);
+		check(InBufferOffset % GRHIRayTracingAccelerationStructureAlignment == 0);
 
-		Layer.View = MakeUnique<FVulkanView>(*Device, VK_DESCRIPTOR_TYPE_ACCELERATION_STRUCTURE_KHR);
-		VkAccelerationStructureKHR NativeAccelerationStructureHandle = Layer.View->InitAsAccelerationStructureView(
+		View = MakeUnique<FVulkanView>(*Device, VK_DESCRIPTOR_TYPE_ACCELERATION_STRUCTURE_KHR);
+		VkAccelerationStructureKHR NativeAccelerationStructureHandle = View->InitAsAccelerationStructureView(
 			AccelerationStructureBuffer
-			, LayerOffset
-			//, Layer.SizeInfo.ResultSize
-			// TODO: Using whole remaining size instead of Layer.SizeInfo.ResultSize reintruduces a validation error but use of Layer.SizeInfo.ResultSize broke RT on Adreno.
-			, InBuffer->GetSize() - LayerOffset 
+			, InBufferOffset
+			//, SizeInfo.ResultSize
+			// TODO: Using whole remaining size instead of SizeInfo.ResultSize reintroduces a validation error but use of SizeInfo.ResultSize broke RT on Adreno.
+			, InBuffer->GetSize() - InBufferOffset
 		)->GetAccelerationStructureView().Handle;
 
 		FString DebugNameString = Initializer.DebugName.ToString();
@@ -792,52 +767,27 @@ void FVulkanRayTracingScene::BuildAccelerationStructure(
 		}
 	}
 
-	PRAGMA_DISABLE_DEPRECATION_WARNINGS
-	const uint32 NumLayers = Initializer.NumNativeInstancesPerLayer.Num();
-	PRAGMA_ENABLE_DEPRECATION_WARNINGS
-
-	TArray<FVkRtTLASBuildData> BuildDatas;
-	BuildDatas.SetNum(NumLayers);
-
-	TArray<VkAccelerationStructureBuildGeometryInfoKHR> GeometryInfos;
-	GeometryInfos.SetNum(NumLayers);
-
-	TArray<VkAccelerationStructureBuildRangeInfoKHR> BuildRanges;
-	BuildRanges.SetNum(NumLayers);
+	FVkRtTLASBuildData BuildData;
 
-	TArray<VkAccelerationStructureBuildRangeInfoKHR*> pBuildRanges;
-	pBuildRanges.SetNum(NumLayers);
+	VkAccelerationStructureBuildRangeInfoKHR BuildRangeInfo;
+	VkAccelerationStructureBuildRangeInfoKHR* pBuildRangeInfos;
 
 	const VkDeviceAddress InstanceBufferAddress = InInstanceBuffer->GetDeviceAddress() + InInstanceOffset;
 
-	uint32 InstanceBaseOffset = 0;
-	for (uint32 LayerIndex = 0; LayerIndex < NumLayers; ++LayerIndex)
 	{
-		const FLayerData& Layer = Layers[LayerIndex];
+		GetTLASBuildData(Device->GetInstanceHandle(), Initializer.NumNativeInstances, InstanceBufferAddress, Initializer.BuildFlags, BuildMode, BuildData);
 
-		FVkRtTLASBuildData& BuildData = BuildDatas[LayerIndex];
-		PRAGMA_DISABLE_DEPRECATION_WARNINGS
-		GetTLASBuildData(Device->GetInstanceHandle(), Initializer.NumNativeInstancesPerLayer[LayerIndex], InstanceBufferAddress, Initializer.BuildFlags, BuildMode, BuildData);
-		PRAGMA_ENABLE_DEPRECATION_WARNINGS
+		checkf(View.IsValid(), TEXT("A buffer must be bound to the ray tracing scene before it can be built."));
+		BuildData.GeometryInfo.dstAccelerationStructure = View->GetAccelerationStructureView().Handle;
+		BuildData.GeometryInfo.srcAccelerationStructure = bIsUpdate ? View->GetAccelerationStructureView().Handle : nullptr;
+		BuildData.GeometryInfo.scratchData.deviceAddress = InScratchBuffer->GetDeviceAddress() + InScratchOffset;
 
-		const uint64 LayerScratchOffset = bIsUpdate ? Layer.UpdateScratchOffset : Layer.BuildScratchOffset;
+		BuildRangeInfo.primitiveCount = Initializer.NumNativeInstances;
+		BuildRangeInfo.primitiveOffset = 0;
+		BuildRangeInfo.transformOffset = 0;
+		BuildRangeInfo.firstVertex = 0;
 
-		checkf(Layer.View.IsValid(), TEXT("A buffer must be bound to the ray tracing scene before it can be built."));
-		BuildData.GeometryInfo.dstAccelerationStructure = Layer.View->GetAccelerationStructureView().Handle;
-		BuildData.GeometryInfo.srcAccelerationStructure = bIsUpdate ? Layer.View->GetAccelerationStructureView().Handle : nullptr;
-		BuildData.GeometryInfo.scratchData.deviceAddress = InScratchBuffer->GetDeviceAddress() + InScratchOffset + LayerScratchOffset;
-
-		GeometryInfos[LayerIndex] = BuildData.GeometryInfo;
-
-		VkAccelerationStructureBuildRangeInfoKHR& TLASBuildRangeInfo = BuildRanges[LayerIndex];
-		PRAGMA_DISABLE_DEPRECATION_WARNINGS
-		TLASBuildRangeInfo.primitiveCount = Initializer.NumNativeInstancesPerLayer[LayerIndex];
-		PRAGMA_ENABLE_DEPRECATION_WARNINGS
-		TLASBuildRangeInfo.primitiveOffset = InstanceBaseOffset;
-		TLASBuildRangeInfo.transformOffset = 0;
-		TLASBuildRangeInfo.firstVertex = 0;
-
-		pBuildRanges[LayerIndex] = &TLASBuildRangeInfo;
+		pBuildRangeInfos = &BuildRangeInfo;
 
 		if (bIsUpdate)
 		{
@@ -847,10 +797,6 @@ void FVulkanRayTracingScene::BuildAccelerationStructure(
 		{
 			INC_DWORD_STAT(STAT_VulkanRayTracingBuiltTLAS);
 		}
-
-		PRAGMA_DISABLE_DEPRECATION_WARNINGS
-		InstanceBaseOffset += Initializer.NumNativeInstancesPerLayer[LayerIndex];
-		PRAGMA_ENABLE_DEPRECATION_WARNINGS
 	}
 
 	FVulkanCommandBufferManager& CommandBufferManager = *CommandContext.GetCommandBufferManager();
@@ -859,7 +805,7 @@ void FVulkanRayTracingScene::BuildAccelerationStructure(
 	// Force a memory barrier to make sure all previous builds ops are finished before building the TLAS
 	AddAccelerationStructureBuildBarrier(CmdBuffer->GetHandle());
 
-	VulkanDynamicAPI::vkCmdBuildAccelerationStructuresKHR(CmdBuffer->GetHandle(), NumLayers, GeometryInfos.GetData(), pBuildRanges.GetData());
+	VulkanDynamicAPI::vkCmdBuildAccelerationStructuresKHR(CmdBuffer->GetHandle(), 1, &BuildData.GeometryInfo, &pBuildRangeInfos);
 
 	// Acceleration structure build barrier is used here to ensure that the acceleration structure build is complete before any rays are traced
 	AddAccelerationStructureBuildBarrier(CmdBuffer->GetHandle());
diff --git a/Engine/Source/Runtime/VulkanRHI/Private/VulkanRayTracing.h b/Engine/Source/Runtime/VulkanRHI/Private/VulkanRayTracing.h
--- a/Engine/Source/Runtime/VulkanRHI/Private/VulkanRayTracing.h
+++ b/Engine/Source/Runtime/VulkanRHI/Private/VulkanRayTracing.h
@@ -202,7 +202,6 @@ public:
 	~FVulkanRayTracingScene();
 
 	const FRayTracingSceneInitializer2& GetInitializer() const override final { return Initializer; }
-	uint32 GetLayerBufferOffset(uint32 LayerIndex) const override final { return Layers[LayerIndex].BufferOffset; }
 
 	void BindBuffer(FRHIBuffer* InBuffer, uint32 InBufferOffset);
 	void BuildAccelerationStructure(
@@ -253,16 +252,7 @@ private:
 	// the lifetime of the scene object may be different from the lifetime of the buffer.
 	// Many VkAccelerationStructureKHR-s may be created, pointing at the same buffer.
 
-	struct FLayerData
-	{
-		TUniquePtr<FVulkanView> View;
-		FRayTracingAccelerationStructureSize SizeInfo;
-		uint32 BufferOffset;
-		uint32 BuildScratchOffset;
-		uint32 UpdateScratchOffset;
-	};
-
-	TArray<FLayerData> Layers;
+	TUniquePtr<FVulkanView> View;
 	
 	TRefCountPtr<FVulkanResourceMultiBuffer> AccelerationStructureBuffer;
 
