Path Tracing Multi-GPU:  Fix for crash in denoiser on D3D12 when the engine is started with multi-GPU enabled.  DirectML (used by the NNE denoiser) requires all resources sent to it to be in single device heaps.  A flag was added to the RHI buffer usage flags to specify a resource is used by NNE, and the transient and pool resource allocators on D3D12 create a separate heap for that case when multi-GPU is active.

#rnx

diff --git a/Engine/Plugins/Experimental/NNEDenoiser/Source/NNEDenoiser/Private/NNEDenoiserGenericDenoiser.cpp b/Engine/Plugins/Experimental/NNEDenoiser/Source/NNEDenoiser/Private/NNEDenoiserGenericDenoiser.cpp
--- a/Engine/Plugins/Experimental/NNEDenoiser/Source/NNEDenoiser/Private/NNEDenoiserGenericDenoiser.cpp
+++ b/Engine/Plugins/Experimental/NNEDenoiser/Source/NNEDenoiser/Private/NNEDenoiserGenericDenoiser.cpp
@@ -39,7 +39,10 @@ FRDGBufferRef CreateBufferRDG(FRDGBuilder& GraphBuilder, const NNE::FTensorDesc&
 	const uint32 BytesPerElement = TensorDesc.GetElementByteSize();
 	const uint32 NumElements = TensorShape.Volume();
 
-	return GraphBuilder.CreateBuffer(FRDGBufferDesc::CreateBufferDesc(BytesPerElement, NumElements), *TensorDesc.GetName());
+	FRDGBufferDesc Desc = FRDGBufferDesc::CreateBufferDesc(BytesPerElement, NumElements);
+	Desc.Usage |= EBufferUsageFlags::NNE;
+
+	return GraphBuilder.CreateBuffer(Desc, *TensorDesc.GetName());
 }
 
 TArray<FRDGBufferRef> CreateBuffersRDG(FRDGBuilder& GraphBuilder, TConstArrayView<NNE::FTensorDesc> TensorDescs, TConstArrayView<NNE::FTensorShape> TensorShapes)
@@ -244,6 +247,8 @@ TUniquePtr<FHistory> FGenericDenoiser::AddPasses(
 	if (TransferFunction.IsValid() && AutoExposure.IsValid())
 	{
 		FRDGBufferDesc InputBufferDesc = FRDGBufferDesc::CreateBufferDesc(sizeof(float), 2);
+		InputBufferDesc.Usage |= EBufferUsageFlags::NNE;
+
 		FRDGBufferRef InputScaleBuffer = GraphBuilder.CreateBuffer(InputBufferDesc, TEXT("AutoExposureOutputBuffer"));
 
 		AutoExposure->EnqueueRDG(GraphBuilder, ColorTex, InputScaleBuffer);
diff --git a/Engine/Source/Runtime/D3D12RHI/Private/D3D12Adapter.cpp b/Engine/Source/Runtime/D3D12RHI/Private/D3D12Adapter.cpp
--- a/Engine/Source/Runtime/D3D12RHI/Private/D3D12Adapter.cpp
+++ b/Engine/Source/Runtime/D3D12RHI/Private/D3D12Adapter.cpp
@@ -792,7 +792,7 @@ FD3D12TransientHeapCache& FD3D12Adapter::GetOrCreateTransientHeapCache()
 {
 	if (!TransientMemoryCache)
 	{
-		TransientMemoryCache = FD3D12TransientHeapCache::Create(this, FRHIGPUMask::All());
+		TransientMemoryCache = FD3D12TransientHeapCache::Create(this);
 	}
 
 	return static_cast<FD3D12TransientHeapCache&>(*TransientMemoryCache);
diff --git a/Engine/Source/Runtime/D3D12RHI/Private/D3D12Allocation.cpp b/Engine/Source/Runtime/D3D12RHI/Private/D3D12Allocation.cpp
--- a/Engine/Source/Runtime/D3D12RHI/Private/D3D12Allocation.cpp
+++ b/Engine/Source/Runtime/D3D12RHI/Private/D3D12Allocation.cpp
@@ -1409,6 +1409,10 @@ FD3D12BufferPool* FD3D12DefaultBufferAllocator::CreateBufferPool(D3D12_HEAP_TYPE
 	FD3D12Device* Device = GetParentDevice();
 	FD3D12ResourceInitConfig InitConfig = FD3D12BufferPool::GetResourceAllocatorInitConfig(InHeapType, InResourceFlags, InBufferUsage);
 
+	// NNE resources must be in heaps visible on GPU0 only.  Required by DirectML.  Note that in single GPU mode, GetVisibilityMask() will
+	// be the same as GPU0(), so no extra heap fragmentation occurs in that case.
+	FRHIGPUMask VisibleNodes = EnumHasAnyFlags(InBufferUsage, EBufferUsageFlags::NNE) ? FRHIGPUMask::GPU0() : GetVisibilityMask();
+
 #if USE_BUFFER_POOL_ALLOCATOR
 
 	const FString Name(L"D3D12 Pool Allocator");
@@ -1435,7 +1439,7 @@ FD3D12BufferPool* FD3D12DefaultBufferAllocator::CreateBufferPool(D3D12_HEAP_TYPE
 	}
 #endif // D3D12_RHI_RAYTRACING
 
-	FD3D12BufferPool* NewPool = new FD3D12PoolAllocator(Device, GetVisibilityMask(), InitConfig, Name, AllocationStrategy, PoolSize, PoolAlignment, MaxAllocationSize, FreeListOrder, bDefragEnabled, TraceHeapId);
+	FD3D12BufferPool* NewPool = new FD3D12PoolAllocator(Device, VisibleNodes, InitConfig, Name, AllocationStrategy, PoolSize, PoolAlignment, MaxAllocationSize, FreeListOrder, bDefragEnabled, TraceHeapId);
 
 #else // USE_BUFFER_POOL_ALLOCATOR
 
@@ -1446,7 +1450,7 @@ FD3D12BufferPool* FD3D12DefaultBufferAllocator::CreateBufferPool(D3D12_HEAP_TYPE
 
 	const FString Name(L"Default Buffer Multi Buddy Allocator");
 	FD3D12MultiBuddyAllocator* Allocator = new FD3D12MultiBuddyAllocator(Device,
-		GetVisibilityMask(),
+		VisibleNodes,
 		InitConfig,
 		Name,
 		AllocationStrategy,
@@ -1512,8 +1516,9 @@ void FD3D12DefaultBufferAllocator::AllocDefaultResource(D3D12_HEAP_TYPE InHeapTy
 	{
 		ResourceLocation.Clear();
 
+		// NNE resources must be in heaps visible on GPU0 only.  Required by DirectML.
 		FD3D12Resource* NewResource = nullptr;
-		const D3D12_HEAP_PROPERTIES HeapProps = CD3DX12_HEAP_PROPERTIES(InHeapType, GetGPUMask().GetNative(), GetVisibilityMask().GetNative());
+		const D3D12_HEAP_PROPERTIES HeapProps = CD3DX12_HEAP_PROPERTIES(InHeapType, GetGPUMask().GetNative(), EnumHasAnyFlags(InBufferUsage, EBufferUsageFlags::NNE) ? FRHIGPUMask::GPU0().GetNative() : GetVisibilityMask().GetNative());
 		D3D12_RESOURCE_DESC Desc = InResourceDesc;
 		Desc.Alignment = 0;
 		VERIFYD3D12RESULT(Adapter->CreateCommittedResource(Desc, GetGPUMask(), HeapProps, InCreateState, InResourceStateMode, InCreateState, nullptr, &NewResource, Name, false));
diff --git a/Engine/Source/Runtime/D3D12RHI/Private/D3D12Buffer.cpp b/Engine/Source/Runtime/D3D12RHI/Private/D3D12Buffer.cpp
--- a/Engine/Source/Runtime/D3D12RHI/Private/D3D12Buffer.cpp
+++ b/Engine/Source/Runtime/D3D12RHI/Private/D3D12Buffer.cpp
@@ -455,7 +455,14 @@ FD3D12Buffer* FD3D12DynamicRHI::CreateD3D12Buffer(class FRHICommandListBase* RHI
 	// Setup the state at which the resource needs to be created - copy dest only supported for placed resources
 	D3D12_RESOURCE_STATES CreateState = (CreateInfo.ResourceArray && bSupportResourceStateTracking) ? D3D12_RESOURCE_STATE_COPY_DEST : DesiredState;
 
-	FD3D12Buffer* Buffer = GetAdapter().CreateRHIBuffer(Desc, Alignment, BufferDesc, StateMode, CreateState, bHasInitialData || bForceKeepUnlocked, CreateInfo.GPUMask, ResourceAllocator, CreateInfo.DebugName, CreateInfo.OwnerName, TraceClassName);
+	// Theoretically, we could assert if GPUMask isn't correct, but at the moment the RDG and RHI buffer descriptions don't include the
+	// GPU mask, so there's no way for the caller to configure it (only the lower level CreateInfo includes it).  Note that differentiation
+	// for NNE (DirectML) is required beyond just setting the mask anyway, in the sense of forcing separate GPU0 visible only heaps, not just
+	// filtering which GPU copies are allocated.  Because this is necessary to solve a crash, it's higher priority than GPUMask support,
+	// which may be added in the future.
+	FRHIGPUMask GPUMask = EnumHasAnyFlags(BufferDesc.Usage, EBufferUsageFlags::NNE) ? FRHIGPUMask::GPU0() : CreateInfo.GPUMask;
+
+	FD3D12Buffer* Buffer = GetAdapter().CreateRHIBuffer(Desc, Alignment, BufferDesc, StateMode, CreateState, bHasInitialData || bForceKeepUnlocked, GPUMask, ResourceAllocator, CreateInfo.DebugName, CreateInfo.OwnerName, TraceClassName);
 	check(Buffer->ResourceLocation.IsValid());
 
 	// Copy the resource data if available 
diff --git a/Engine/Source/Runtime/D3D12RHI/Private/D3D12PoolAllocator.cpp b/Engine/Source/Runtime/D3D12RHI/Private/D3D12PoolAllocator.cpp
--- a/Engine/Source/Runtime/D3D12RHI/Private/D3D12PoolAllocator.cpp
+++ b/Engine/Source/Runtime/D3D12RHI/Private/D3D12PoolAllocator.cpp
@@ -252,6 +252,12 @@ FD3D12PoolAllocator::~FD3D12PoolAllocator()
 
 bool FD3D12PoolAllocator::SupportsAllocation(D3D12_HEAP_TYPE InHeapType, D3D12_RESOURCE_FLAGS InResourceFlags, EBufferUsageFlags InBufferUsage, ED3D12ResourceStateMode InResourceStateMode, uint32 Alignment) const
 {
+	// NNE resources must be in heaps visible on GPU0 only.  Required by DirectML.
+	if (EnumHasAnyFlags(InBufferUsage, EBufferUsageFlags::NNE) && VisibilityMask != FRHIGPUMask::GPU0())
+	{
+		return false;
+	}
+
 	FD3D12ResourceInitConfig InInitConfig = GetResourceAllocatorInitConfig(InHeapType, InResourceFlags, InBufferUsage);
 	EResourceAllocationStrategy InAllocationStrategy = GetResourceAllocationStrategy(InResourceFlags, InResourceStateMode, Alignment);
 
diff --git a/Engine/Source/Runtime/D3D12RHI/Private/D3D12TransientResourceAllocator.cpp b/Engine/Source/Runtime/D3D12RHI/Private/D3D12TransientResourceAllocator.cpp
--- a/Engine/Source/Runtime/D3D12RHI/Private/D3D12TransientResourceAllocator.cpp
+++ b/Engine/Source/Runtime/D3D12RHI/Private/D3D12TransientResourceAllocator.cpp
@@ -127,24 +127,26 @@ FD3D12TransientHeap::~FD3D12TransientHeap()
 	}
 }
 
-TUniquePtr<FD3D12TransientHeapCache> FD3D12TransientHeapCache::Create(FD3D12Adapter* ParentAdapter, FRHIGPUMask VisibleNodeMask)
+TUniquePtr<FD3D12TransientHeapCache> FD3D12TransientHeapCache::Create(FD3D12Adapter* ParentAdapter)
 {
 	FRHITransientHeapCache::FInitializer Initializer = FRHITransientHeapCache::FInitializer::CreateDefault();
 
 	// Tier2 hardware is able to mix resource types onto the same heap.
 	Initializer.bSupportsAllHeapFlags = ParentAdapter->GetResourceHeapTier() == D3D12_RESOURCE_HEAP_TIER_2;
 
-	return TUniquePtr<FD3D12TransientHeapCache>(new FD3D12TransientHeapCache(Initializer, ParentAdapter, VisibleNodeMask));
+	return TUniquePtr<FD3D12TransientHeapCache>(new FD3D12TransientHeapCache(Initializer, ParentAdapter));
 }
 
-FD3D12TransientHeapCache::FD3D12TransientHeapCache(const FRHITransientHeapCache::FInitializer& Initializer, FD3D12Adapter* ParentAdapter, FRHIGPUMask InVisibleNodeMask)
+FD3D12TransientHeapCache::FD3D12TransientHeapCache(const FRHITransientHeapCache::FInitializer& Initializer, FD3D12Adapter* ParentAdapter)
 	: FRHITransientHeapCache(Initializer)
 	, FD3D12AdapterChild(ParentAdapter)
-	, VisibleNodeMask(InVisibleNodeMask)
 {}
 
 FRHITransientHeap* FD3D12TransientHeapCache::CreateHeap(const FRHITransientHeap::FInitializer& HeapInitializer)
 {
+	// If heap is flagged for NNE buffers, make it visible on first GPU only.  Required by DirectML.
+	FRHIGPUMask VisibleNodeMask = HeapInitializer.Flags == ERHITransientHeapFlags::AllowNNEBuffers ? FRHIGPUMask::GPU0() : FRHIGPUMask::All();
+
 	return GetParentAdapter()->CreateLinkedObject<FD3D12TransientHeap>(VisibleNodeMask, [&](FD3D12Device* Device)
 	{
 		return new FD3D12TransientHeap(HeapInitializer, GetParentAdapter(), Device, VisibleNodeMask);
diff --git a/Engine/Source/Runtime/D3D12RHI/Private/D3D12TransientResourceAllocator.h b/Engine/Source/Runtime/D3D12RHI/Private/D3D12TransientResourceAllocator.h
--- a/Engine/Source/Runtime/D3D12RHI/Private/D3D12TransientResourceAllocator.h
+++ b/Engine/Source/Runtime/D3D12RHI/Private/D3D12TransientResourceAllocator.h
@@ -27,15 +27,13 @@ class FD3D12TransientHeapCache final
 	, public FD3D12AdapterChild
 {
 public:
-	static TUniquePtr<FD3D12TransientHeapCache> Create(FD3D12Adapter* ParentAdapter, FRHIGPUMask VisibleNodeMask);
+	static TUniquePtr<FD3D12TransientHeapCache> Create(FD3D12Adapter* ParentAdapter);
 
 	//! FRHITransientResourceSystem Overrides
 	FRHITransientHeap* CreateHeap(const FRHITransientHeap::FInitializer& Initializer) override;
 
 private:
-	FD3D12TransientHeapCache(const FInitializer& Initializer, FD3D12Adapter* ParentAdapter, FRHIGPUMask VisibleNodeMask);
-
-	FRHIGPUMask VisibleNodeMask;
+	FD3D12TransientHeapCache(const FInitializer& Initializer, FD3D12Adapter* ParentAdapter);
 };
 
 class FD3D12TransientResourceHeapAllocator final
diff --git a/Engine/Source/Runtime/RHI/Public/RHIDefinitions.h b/Engine/Source/Runtime/RHI/Public/RHIDefinitions.h
--- a/Engine/Source/Runtime/RHI/Public/RHIDefinitions.h
+++ b/Engine/Source/Runtime/RHI/Public/RHIDefinitions.h
@@ -846,6 +846,9 @@ enum class EBufferUsageFlags : uint32
 	/** Buffer should go in fast vram (hint only). Requires BUF_Transient */
 	FastVRAM                = 1 << 10,
 
+	/** Buffer is used by NNE.  DirectML requires NNE resources to be in single device memory heaps when multi-GPU is active. */
+	NNE						= 1 << 11,
+
 	/** Create a buffer that can be shared with an external RHI or process. */
 	Shared                  = 1 << 12,
 
diff --git a/Engine/Source/Runtime/RHICore/Private/RHICoreTransientResourceAllocator.cpp b/Engine/Source/Runtime/RHICore/Private/RHICoreTransientResourceAllocator.cpp
--- a/Engine/Source/Runtime/RHICore/Private/RHICoreTransientResourceAllocator.cpp
+++ b/Engine/Source/Runtime/RHICore/Private/RHICoreTransientResourceAllocator.cpp
@@ -557,7 +557,18 @@ FRHITransientHeap* FRHITransientHeapCache::Acquire(uint64 FirstAllocationSize, E
 	FRHITransientHeap::FInitializer HeapInitializer;
 	HeapInitializer.Size = GetHeapSize(FirstAllocationSize);
 	HeapInitializer.Alignment = Initializer.HeapAlignment;
-	HeapInitializer.Flags = (Initializer.bSupportsAllHeapFlags ? ERHITransientHeapFlags::AllowAll : FirstAllocationHeapFlags);
+	if (GNumExplicitGPUsForRendering > 1)
+	{
+		// With multi-GPU, we need separate GPU0 only heaps for NNE accessible buffers.  Required by DirectML.  Note that the calling
+		// code only sets one flag for a given allocation, so if the flag is NNE, create a heap with that flag alone, otherwise create
+		// a heap with the rest of the flags (or if bSupportsAllHeapFlags is false, that also forces a single flag per heap).
+		HeapInitializer.Flags = Initializer.bSupportsAllHeapFlags && FirstAllocationHeapFlags != ERHITransientHeapFlags::AllowNNEBuffers ?
+			ERHITransientHeapFlags::AllowBuffers | ERHITransientHeapFlags::AllowTextures | ERHITransientHeapFlags::AllowRenderTargets : FirstAllocationHeapFlags;
+	}
+	else
+	{
+		HeapInitializer.Flags = Initializer.bSupportsAllHeapFlags ? ERHITransientHeapFlags::AllowAll : FirstAllocationHeapFlags;
+	}
 	HeapInitializer.TextureCacheSize = Initializer.TextureCacheSize;
 	HeapInitializer.BufferCacheSize = Initializer.BufferCacheSize;
 
@@ -671,9 +682,11 @@ FRHITransientBuffer* FRHITransientResourceHeapAllocator::CreateBufferInternal(
 {
 	FRHITransientBuffer* Buffer = nullptr;
 
+	ERHITransientHeapFlags BufferHeapFlag = (GNumExplicitGPUsForRendering > 1) && EnumHasAnyFlags(CreateInfo.Usage, EBufferUsageFlags::NNE) ? ERHITransientHeapFlags::AllowNNEBuffers : ERHITransientHeapFlags::AllowBuffers;
+
 	for (FRHITransientHeap* Heap : Heaps)
 	{
-		if (!Heap->IsAllocationSupported(BufferSize, ERHITransientHeapFlags::AllowBuffers))
+		if (!Heap->IsAllocationSupported(BufferSize, BufferHeapFlag))
 		{
 			continue;
 		}
@@ -688,7 +701,7 @@ FRHITransientBuffer* FRHITransientResourceHeapAllocator::CreateBufferInternal(
 
 	if (!Buffer)
 	{
-		FRHITransientHeap* Heap = HeapCache.Acquire(BufferSize, ERHITransientHeapFlags::AllowBuffers);
+		FRHITransientHeap* Heap = HeapCache.Acquire(BufferSize, BufferHeapFlag);
 		Heaps.Emplace(Heap);
 
 		Buffer = Heap->CreateBuffer(CreateInfo, DebugName, Fences, CurrentCycle, BufferSize, BufferAlignment, CreateBufferFunction);
diff --git a/Engine/Source/Runtime/RHICore/Public/RHICoreTransientResourceAllocator.h b/Engine/Source/Runtime/RHICore/Public/RHICoreTransientResourceAllocator.h
--- a/Engine/Source/Runtime/RHICore/Public/RHICoreTransientResourceAllocator.h
+++ b/Engine/Source/Runtime/RHICore/Public/RHICoreTransientResourceAllocator.h
@@ -442,8 +442,11 @@ enum class ERHITransientHeapFlags : uint8
 	// Supports placing render targets onto the heap.
 	AllowRenderTargets = 1 << 2,
 
+	// Supports placing NNE accessible buffers onto the heap.  Differentiation is required for DirectML with multi-GPU.
+	AllowNNEBuffers = 1 << 3,
+
 	// Supports all resource types.
-	AllowAll = AllowBuffers | AllowTextures | AllowRenderTargets
+	AllowAll = AllowBuffers | AllowTextures | AllowRenderTargets | AllowNNEBuffers
 };
 
 ENUM_CLASS_FLAGS(ERHITransientHeapFlags);
