[Persistent FRHIRayTracingScene] Specify MaxNumInstances in FRayTracingSceneInitializer2 and actual number of instances in FRayTracingSceneBuildParams.

diff --git a/Engine/Plugins/Experimental/GPULightmass/Source/GPULightmass/Private/LightmapRenderer.cpp b/Engine/Plugins/Experimental/GPULightmass/Source/GPULightmass/Private/LightmapRenderer.cpp
--- a/Engine/Plugins/Experimental/GPULightmass/Source/GPULightmass/Private/LightmapRenderer.cpp
+++ b/Engine/Plugins/Experimental/GPULightmass/Source/GPULightmass/Private/LightmapRenderer.cpp
@@ -1140,7 +1140,7 @@ bool FSceneRenderState::SetupRayTracingScene(FRDGBuilder& GraphBuilder, FSceneUn
 			{
 				FRayTracingSceneInitializer2 Initializer;
 				Initializer.DebugName = FName(TEXT("LightmapRendererRayTracingScene"));
-				Initializer.NumNativeInstances = RayTracingSceneInitializationData.NumNativeGPUSceneInstances + RayTracingSceneInitializationData.NumNativeCPUInstances;
+				Initializer.MaxNumInstances = RayTracingSceneInitializationData.NumNativeGPUSceneInstances + RayTracingSceneInitializationData.NumNativeCPUInstances;
 				Initializer.NumTotalSegments = RayTracingSceneInitializationData.TotalNumSegments;
 				Initializer.BuildFlags = SceneBuildFlags;
 
@@ -1149,7 +1149,7 @@ bool FSceneRenderState::SetupRayTracingScene(FRDGBuilder& GraphBuilder, FSceneUn
 
 			const FRayTracingSceneInitializer2& SceneInitializer = RayTracingScene->GetInitializer();
 
-			FRayTracingAccelerationStructureSize SizeInfo = RHICalcRayTracingSceneSize(SceneInitializer.NumNativeInstances, SceneBuildFlags);
+			FRayTracingAccelerationStructureSize SizeInfo = RHICalcRayTracingSceneSize(SceneInitializer.MaxNumInstances, SceneBuildFlags);
 			FRHIResourceCreateInfo BufferCreateInfo(TEXT("LightmassRayTracingSceneBuffer"));
 			RayTracingSceneBuffer = RHICmdList.CreateBuffer(uint32(SizeInfo.ResultSize), BUF_AccelerationStructure, 0, ERHIAccess::BVHWrite, BufferCreateInfo);
 			RayTracingSceneSRV = RHICmdList.CreateShaderResourceView(FShaderResourceViewInitializer(RayTracingSceneBuffer, RayTracingScene, 0));
@@ -1161,7 +1161,7 @@ bool FSceneRenderState::SetupRayTracingScene(FRDGBuilder& GraphBuilder, FSceneUn
 				ScratchBufferCreateInfo);
 
 			FRWBufferStructured InstanceBuffer;
-			InstanceBuffer.Initialize(RHICmdList, TEXT("LightmassRayTracingInstanceBuffer"), GRHIRayTracingInstanceDescriptorSize, SceneInitializer.NumNativeInstances);
+			InstanceBuffer.Initialize(RHICmdList, TEXT("LightmassRayTracingInstanceBuffer"), GRHIRayTracingInstanceDescriptorSize, SceneInitializer.MaxNumInstances);
 
 			// Need to pass "BUF_MultiGPUAllocate", as acceleration structure virtual addresses are different per GPU
 			FByteAddressBuffer AccelerationStructureAddressesBuffer;
@@ -1171,7 +1171,7 @@ bool FSceneRenderState::SetupRayTracingScene(FRDGBuilder& GraphBuilder, FSceneUn
 				RayTracingSceneInitializationData.ReferencedGeometries.Num() * sizeof(FRayTracingAccelerationStructureAddress),
 				BUF_Volatile | BUF_MultiGPUAllocate);
 
-			const uint32 InstanceUploadBufferSize = SceneInitializer.NumNativeInstances * sizeof(FRayTracingInstanceDescriptorInput);
+			const uint32 InstanceUploadBufferSize = SceneInitializer.MaxNumInstances * sizeof(FRayTracingInstanceDescriptorInput);
 			FBufferRHIRef InstanceUploadBuffer;
 			FShaderResourceViewRHIRef InstanceUploadSRV;
 			{
@@ -1201,7 +1201,7 @@ bool FSceneRenderState::SetupRayTracingScene(FRDGBuilder& GraphBuilder, FSceneUn
 					RayTracingSceneInitializationData.BaseInstancePrefixSum,
 					RayTracingSceneInitializationData.NumNativeGPUSceneInstances,
 					RayTracingSceneInitializationData.NumNativeCPUInstances,
-					MakeArrayView(InstanceUploadData, SceneInitializer.NumNativeInstances),
+					MakeArrayView(InstanceUploadData, SceneInitializer.MaxNumInstances),
 					MakeArrayView(TransformUploadData, RayTracingSceneInitializationData.NumNativeCPUInstances * 3));
 				RHICmdList.UnlockBuffer(TransformUploadBuffer);
 				RHICmdList.UnlockBuffer(InstanceUploadBuffer);
@@ -1255,6 +1255,7 @@ bool FSceneRenderState::SetupRayTracingScene(FRDGBuilder& GraphBuilder, FSceneUn
 				BuildParams.InstanceBufferOffset = 0;
 				BuildParams.ReferencedGeometries = RayTracingSceneInitializationData.ReferencedGeometries;
 				BuildParams.PerInstanceGeometries = RayTracingSceneInitializationData.PerInstanceGeometries;
+				BuildParams.NumInstances = RayTracingSceneInitializationData.NumNativeGPUSceneInstances + RayTracingSceneInitializationData.NumNativeCPUInstances;
 
 				RHICmdList.BuildAccelerationStructure(BuildParams);
 			}
diff --git a/Engine/Source/Runtime/D3D12RHI/Private/D3D12RayTracing.cpp b/Engine/Source/Runtime/D3D12RHI/Private/D3D12RayTracing.cpp
--- a/Engine/Source/Runtime/D3D12RHI/Private/D3D12RayTracing.cpp
+++ b/Engine/Source/Runtime/D3D12RHI/Private/D3D12RayTracing.cpp
@@ -3396,14 +3396,8 @@ FD3D12RayTracingScene::FD3D12RayTracingScene(FD3D12Adapter* Adapter, FRayTracing
 	PRAGMA_ENABLE_DEPRECATION_WARNINGS
 	checkf(Initializer.Lifetime == RTSL_SingleFrame, TEXT("Only single-frame ray tracing scenes are currently implemented."));
 
-	BuildInputs = {};
-	BuildInputs.Type = D3D12_RAYTRACING_ACCELERATION_STRUCTURE_TYPE_TOP_LEVEL;
-	BuildInputs.DescsLayout = D3D12_ELEMENTS_LAYOUT_ARRAY;
-	BuildInputs.NumDescs = Initializer.NumNativeInstances;
-	BuildInputs.Flags = TranslateRayTracingAccelerationStructureFlags(Initializer.BuildFlags);
-
 	// Get maximum buffer sizes for all GPUs in the system
-	SizeInfo = RHICalcRayTracingSceneSize(BuildInputs.NumDescs, Initializer.BuildFlags);
+	SizeInfo = RHICalcRayTracingSceneSize(Initializer.MaxNumInstances, Initializer.BuildFlags);
 };
 
 FD3D12RayTracingScene::~FD3D12RayTracingScene()
@@ -3448,25 +3442,37 @@ void FD3D12RayTracingScene::BindBuffer(FRHIBuffer* InBuffer, uint32 InBufferOffs
 	BufferOffset = InBufferOffset;
 }
 
-void FD3D12RayTracingScene::BuildAccelerationStructure(FD3D12CommandContext& CommandContext,
+void BuildAccelerationStructure(FD3D12CommandContext& CommandContext,
+	FD3D12RayTracingScene& Scene,
 	FD3D12Buffer* ScratchBuffer, uint32 ScratchBufferOffset,
 	FD3D12Buffer* InstanceBuffer, uint32 InstanceBufferOffset,
+	uint32 NumInstances,
 	EAccelerationStructureBuildMode BuildMode)
 {
-	check(InstanceBuffer != nullptr);
-
 	TRACE_CPUPROFILER_EVENT_SCOPE(BuildAccelerationStructure_TopLevel);
 	SCOPE_CYCLE_COUNTER(STAT_D3D12BuildTLAS);
 
+	check(InstanceBuffer != nullptr);
+	checkf(NumInstances <= Scene.Initializer.MaxNumInstances, TEXT("NumInstances must be less or equal to MaxNumInstances"));
+
 	const bool bIsUpdate = BuildMode == EAccelerationStructureBuildMode::Update;
 
+	if (bIsUpdate)
+	{
+		checkf(NumInstances == Scene.NumInstances, TEXT("Number of instances used to update TLAS must match the number used to build."));
+	}
+	else
+	{
+		Scene.NumInstances = NumInstances;
+	}
+
 	const uint32 GPUIndex = CommandContext.GetGPUIndex();
 	FD3D12Adapter* Adapter = CommandContext.GetParentAdapter();
 
 	TRefCountPtr<FD3D12Buffer> AutoScratchBuffer;
 	if (ScratchBuffer == nullptr)
 	{
-		const uint64 ScratchBufferSize = bIsUpdate ? SizeInfo.UpdateScratchSize : SizeInfo.BuildScratchSize;
+		const uint64 ScratchBufferSize = bIsUpdate ? Scene.SizeInfo.UpdateScratchSize : Scene.SizeInfo.BuildScratchSize;
 
 		static const FName ScratchBufferName("AutoBuildScratchTLAS");
 		AutoScratchBuffer = CreateRayTracingBuffer(Adapter, GPUIndex, ScratchBufferSize, ERayTracingBufferType::Scratch, ScratchBufferName);
@@ -3474,42 +3480,52 @@ void FD3D12RayTracingScene::BuildAccelerationStructure(FD3D12CommandContext& Com
 		ScratchBufferOffset = 0;
 	}
 
-	checkf(ScratchBuffer, TEXT("TLAS build requires scratch buffer of at least %lld bytes."), SizeInfo.BuildScratchSize);
+	if (bIsUpdate)
+	{
+		checkf(ScratchBuffer, TEXT("TLAS update requires scratch buffer of at least %lld bytes."), Scene.SizeInfo.UpdateScratchSize);
+	}
+	else
+	{
+		checkf(ScratchBuffer, TEXT("TLAS build requires scratch buffer of at least %lld bytes."), Scene.SizeInfo.BuildScratchSize);
+	}
 
 	{
+		D3D12_BUILD_RAYTRACING_ACCELERATION_STRUCTURE_INPUTS BuildInputs;
+		BuildInputs = {};
+		BuildInputs.Type = D3D12_RAYTRACING_ACCELERATION_STRUCTURE_TYPE_TOP_LEVEL;
+		BuildInputs.DescsLayout = D3D12_ELEMENTS_LAYOUT_ARRAY;
+		BuildInputs.NumDescs = NumInstances;
+		BuildInputs.Flags = TranslateRayTracingAccelerationStructureFlags(Scene.Initializer.BuildFlags);
+
 		D3D12_RAYTRACING_ACCELERATION_STRUCTURE_PREBUILD_INFO PrebuildInfo = {};
 
 		CommandContext.GetParentDevice()->GetRaytracingAccelerationStructurePrebuildInfo(&BuildInputs, &PrebuildInfo);
 
-		checkf(PrebuildInfo.ResultDataMaxSizeInBytes <= SizeInfo.ResultSize,
+		checkf(PrebuildInfo.ResultDataMaxSizeInBytes <= Scene.SizeInfo.ResultSize,
 			TEXT("TLAS build result buffer now requires %lld bytes, but only %lld was calculated in the constructor."),
-			PrebuildInfo.ResultDataMaxSizeInBytes, SizeInfo.ResultSize);
+			PrebuildInfo.ResultDataMaxSizeInBytes, Scene.SizeInfo.ResultSize);
 
-		checkf(PrebuildInfo.ScratchDataSizeInBytes <= SizeInfo.BuildScratchSize,
+		checkf(PrebuildInfo.ScratchDataSizeInBytes <= Scene.SizeInfo.BuildScratchSize,
 			TEXT("TLAS build scratch buffer now requires %lld bytes, but only %lld was calculated in the constructor."),
-			PrebuildInfo.ScratchDataSizeInBytes, SizeInfo.BuildScratchSize);
+			PrebuildInfo.ScratchDataSizeInBytes, Scene.SizeInfo.BuildScratchSize);
 
-		checkf(PrebuildInfo.UpdateScratchDataSizeInBytes <= SizeInfo.UpdateScratchSize,
+		checkf(PrebuildInfo.UpdateScratchDataSizeInBytes <= Scene.SizeInfo.UpdateScratchSize,
 			TEXT("TLAS update scratch buffer now requires %lld bytes, but only %lld was calculated in the constructor."),
-			PrebuildInfo.UpdateScratchDataSizeInBytes, SizeInfo.UpdateScratchSize);
+			PrebuildInfo.UpdateScratchDataSizeInBytes, Scene.SizeInfo.UpdateScratchSize);
 
 		if (bIsUpdate)
 		{
 			checkf(ScratchBufferOffset + PrebuildInfo.UpdateScratchDataSizeInBytes <= ScratchBuffer->GetSize(),
-				TEXT("TLAS scratch buffer size is %d bytes with offset %d (%d bytes available), but the update requires %lld bytes. ")
-				TEXT("BuildInputs.NumDescs = %d, Instances.Num = %d."),
+				TEXT("TLAS scratch buffer size is %d bytes with offset %d (%d bytes available), but the update requires %lld bytes (NumInstances = %d)."),
 				ScratchBuffer->GetSize(), ScratchBufferOffset, ScratchBuffer->GetSize() - ScratchBufferOffset,
-				PrebuildInfo.UpdateScratchDataSizeInBytes,
-				BuildInputs.NumDescs, Initializer.NumNativeInstances);
+				PrebuildInfo.UpdateScratchDataSizeInBytes, NumInstances);
 		}
 		else
 		{
 			checkf(ScratchBufferOffset + PrebuildInfo.ScratchDataSizeInBytes <= ScratchBuffer->GetSize(),
-				TEXT("TLAS scratch buffer size is %d bytes with offset %d (%d bytes available), but the build requires %lld bytes. ")
-				TEXT("BuildInputs.NumDescs = %d, Instances.Num = %d."),
+				TEXT("TLAS scratch buffer size is %d bytes with offset %d (%d bytes available), but the build requires %lld bytes (NumInstances = %d)."),
 				ScratchBuffer->GetSize(), ScratchBufferOffset, ScratchBuffer->GetSize() - ScratchBufferOffset,
-				PrebuildInfo.ScratchDataSizeInBytes,
-				BuildInputs.NumDescs, Initializer.NumNativeInstances);
+				PrebuildInfo.ScratchDataSizeInBytes, NumInstances);
 		}
 	}
 
@@ -3519,7 +3535,7 @@ void FD3D12RayTracingScene::BuildAccelerationStructure(FD3D12CommandContext& Com
 		CommandContext.UpdateResidency(InstanceBuffer->GetResource());
 
 		{
-			TArray<const FD3D12Resource*>& ResourcesToMakeResidentForThisGPU = ResourcesToMakeResident[GPUIndex];
+			TArray<const FD3D12Resource*>& ResourcesToMakeResidentForThisGPU = Scene.ResourcesToMakeResident[GPUIndex];
 
 			ResourcesToMakeResidentForThisGPU.Reset(0);
 
@@ -3561,10 +3577,10 @@ void FD3D12RayTracingScene::BuildAccelerationStructure(FD3D12CommandContext& Com
 			#endif // ENABLE_RESIDENCY_MANAGEMENT
 			};
 
-			const int32 NumReferencedGeometries = ReferencedGeometries.Num();
+			const int32 NumReferencedGeometries = Scene.ReferencedGeometries.Num();
 			for (int32 Index = 0; Index < NumReferencedGeometries; ++Index)
 			{
-				FD3D12RayTracingGeometry* Geometry = FD3D12DynamicRHI::ResourceCast(ReferencedGeometries[Index].GetReference());
+				FD3D12RayTracingGeometry* Geometry = FD3D12DynamicRHI::ResourceCast(Scene.ReferencedGeometries[Index].GetReference());
 
 				checkf(!Geometry->IsDirty(CommandContext.GetGPUIndex()),
 					TEXT("Acceleration structures for all geometries must be built before building the top level acceleration structure for the scene."));
@@ -3593,14 +3609,14 @@ void FD3D12RayTracingScene::BuildAccelerationStructure(FD3D12CommandContext& Com
 	
 	// Build the actual acceleration structure
 
-	const int32 NumReferencedGeometries = ReferencedGeometries.Num();
+	const int32 NumReferencedGeometries = Scene.ReferencedGeometries.Num();
 	for (int32 Index = 0; Index < NumReferencedGeometries; ++Index)
 	{
-		FD3D12RayTracingGeometry* Geometry = FD3D12DynamicRHI::ResourceCast(ReferencedGeometries[Index].GetReference());
+		FD3D12RayTracingGeometry* Geometry = FD3D12DynamicRHI::ResourceCast(Scene.ReferencedGeometries[Index].GetReference());
 		CommandContext.UpdateResidency(Geometry->AccelerationStructureBuffers[GPUIndex]->ResourceLocation.GetResource());
 	}
 
-	TRefCountPtr<FD3D12Buffer>& AccelerationStructureBuffer = AccelerationStructureBuffers[GPUIndex];
+	TRefCountPtr<FD3D12Buffer>& AccelerationStructureBuffer = Scene.AccelerationStructureBuffers[GPUIndex];
 	checkf(AccelerationStructureBuffer.IsValid(), 
 		TEXT("Acceleration structure buffer must be set for this scene using RHIBindAccelerationStructureMemory() before build command is issued."));
 
@@ -3621,10 +3637,10 @@ void FD3D12RayTracingScene::BuildAccelerationStructure(FD3D12CommandContext& Com
 		uint32 InstanceBufferStride = GRHIRayTracingInstanceDescriptorSize;
 		PRAGMA_DISABLE_DEPRECATION_WARNINGS
 		// TODO: Validation related to SBT needs to be done somewhere else since SBT is not known when in BuildAccelerationStructure
-		uint32 TotalHitGroupSlots = Initializer.NumTotalSegments * Initializer.ShaderSlotsPerGeometrySegment;
+		uint32 TotalHitGroupSlots = Scene.Initializer.NumTotalSegments * Scene.Initializer.ShaderSlotsPerGeometrySegment;
 		PRAGMA_ENABLE_DEPRECATION_WARNINGS
 		FRayTracingValidateSceneBuildParamsCS::Dispatch(RHICmdList,
-			TotalHitGroupSlots, BuildInputs.NumDescs,
+			TotalHitGroupSlots, NumInstances,
 			InstanceBuffer, InstanceBufferOffset, InstanceBufferStride);
 	}
 
@@ -3636,7 +3652,7 @@ void FD3D12RayTracingScene::BuildAccelerationStructure(FD3D12CommandContext& Com
 	BuildDescs.Reserve(1);
 
 	{
-		const D3D12_GPU_VIRTUAL_ADDRESS BufferAddress = AccelerationStructureBuffer->ResourceLocation.GetGPUVirtualAddress() + BufferOffset;
+		const D3D12_GPU_VIRTUAL_ADDRESS BufferAddress = AccelerationStructureBuffer->ResourceLocation.GetGPUVirtualAddress() + Scene.BufferOffset;
 		D3D12_GPU_VIRTUAL_ADDRESS ScratchAddress = ScratchBuffer->ResourceLocation.GetGPUVirtualAddress() + ScratchBufferOffset;
 
 		checkf(BufferAddress % GRHIRayTracingAccelerationStructureAlignment == 0,
@@ -3648,12 +3664,16 @@ void FD3D12RayTracingScene::BuildAccelerationStructure(FD3D12CommandContext& Com
 			GRHIRayTracingScratchBufferAlignment);
 
 		D3D12_BUILD_RAYTRACING_ACCELERATION_STRUCTURE_DESC& BuildDesc = BuildDescs.AddDefaulted_GetRef();
-		BuildDesc.Inputs = BuildInputs;
+		BuildDesc.Inputs = {};
+		BuildDesc.Inputs.Type = D3D12_RAYTRACING_ACCELERATION_STRUCTURE_TYPE_TOP_LEVEL;
+		BuildDesc.Inputs.DescsLayout = D3D12_ELEMENTS_LAYOUT_ARRAY;
+		BuildDesc.Inputs.NumDescs = NumInstances;
 		BuildDesc.Inputs.InstanceDescs = InstanceBuffer->ResourceLocation.GetGPUVirtualAddress() + InstanceBufferOffset;
+		BuildDesc.Inputs.Flags = TranslateRayTracingAccelerationStructureFlags(Scene.Initializer.BuildFlags);
 
 		if (bIsUpdate)
 		{
-			checkf(EnumHasAllFlags(Initializer.BuildFlags, ERayTracingAccelerationStructureFlags::AllowUpdate),
+			checkf(EnumHasAllFlags(Scene.Initializer.BuildFlags, ERayTracingAccelerationStructureFlags::AllowUpdate),
 				TEXT("Acceleration structure must be created with FRayTracingGeometryInitializer::bAllowUpdate=true to perform refit / update."));
 
 			BuildDesc.Inputs.Flags |= D3D12_RAYTRACING_ACCELERATION_STRUCTURE_BUILD_FLAG_PERFORM_UPDATE;
@@ -3679,10 +3699,10 @@ void FD3D12RayTracingScene::BuildAccelerationStructure(FD3D12CommandContext& Com
 	// #dxr_todo: these barriers should ideally be inserted by the high level code to allow more overlapped execution
 	CommandContext.AddUAVBarrier();
 
-	bBuilt = true;
+	Scene.bBuilt = true;
 
 #if D3D12_RHI_SUPPORT_RAYTRACING_SCENE_DEBUGGING
-	D3D12RayTracingSceneDebugUpdate(*this, InstanceBuffer, InstanceBufferOffset, CommandContext);
+	D3D12RayTracingSceneDebugUpdate(Scene, InstanceBuffer, InstanceBufferOffset, CommandContext);
 #endif // D3D12_RHI_SUPPORT_RAYTRACING_SCENE_DEBUGGING
 }
 
@@ -3941,10 +3961,12 @@ void FD3D12CommandContext::RHIBuildAccelerationStructure(const FRayTracingSceneB
 		Scene->ReferencedGeometries.Add(ReferencedGeometry);
 	}
 
-	Scene->BuildAccelerationStructure(
+	BuildAccelerationStructure(
 		*this,
+		*Scene,
 		ScratchBuffer, SceneBuildParams.ScratchBufferOffset,
 		InstanceBuffer, SceneBuildParams.InstanceBufferOffset,
+		SceneBuildParams.NumInstances,
 		SceneBuildParams.BuildMode
 	);
 }
diff --git a/Engine/Source/Runtime/D3D12RHI/Private/D3D12RayTracing.h b/Engine/Source/Runtime/D3D12RHI/Private/D3D12RayTracing.h
--- a/Engine/Source/Runtime/D3D12RHI/Private/D3D12RayTracing.h
+++ b/Engine/Source/Runtime/D3D12RHI/Private/D3D12RayTracing.h
@@ -103,13 +103,9 @@ public:
 	void BindBuffer(FRHIBuffer* Buffer, uint32 BufferOffset);
 	void ReleaseBuffer();
 
-	void BuildAccelerationStructure(FD3D12CommandContext& CommandContext,
-		FD3D12Buffer* ScratchBuffer, uint32 ScratchBufferOffset,
-		FD3D12Buffer* InstanceBuffer, uint32 InstanceBufferOffset,
-		EAccelerationStructureBuildMode BuildMode
-	);
+	using FRHIRayTracingAccelerationStructure::SizeInfo;
 
-	D3D12_BUILD_RAYTRACING_ACCELERATION_STRUCTURE_INPUTS BuildInputs = {};
+	uint32 NumInstances = 0;
 
 	TRefCountPtr<FD3D12Buffer> AccelerationStructureBuffers[MAX_NUM_GPUS];
 	uint32 BufferOffset = 0;
diff --git a/Engine/Source/Runtime/D3D12RHI/Private/D3D12RayTracingDebug.cpp b/Engine/Source/Runtime/D3D12RHI/Private/D3D12RayTracingDebug.cpp
--- a/Engine/Source/Runtime/D3D12RHI/Private/D3D12RayTracingDebug.cpp
+++ b/Engine/Source/Runtime/D3D12RHI/Private/D3D12RayTracingDebug.cpp
@@ -155,7 +155,7 @@ static void DebugSerializeScene(const FD3D12RayTracingScene& Scene, FD3D12Buffer
 	}
 
 	{
-		SceneHeader.NumInstances = Scene.BuildInputs.NumDescs;
+		SceneHeader.NumInstances = Scene.NumInstances;
 
 		// Instance buffer
 
diff --git a/Engine/Source/Runtime/RHI/Public/DynamicRHI.h b/Engine/Source/Runtime/RHI/Public/DynamicRHI.h
--- a/Engine/Source/Runtime/RHI/Public/DynamicRHI.h
+++ b/Engine/Source/Runtime/RHI/Public/DynamicRHI.h
@@ -1458,7 +1458,7 @@ FORCEINLINE FRayTracingSceneRHIRef RHICreateRayTracingScene(FRayTracingSceneInit
 	PRAGMA_DISABLE_DEPRECATION_WARNINGS
 	if (Initializer.NumNativeInstancesPerLayer.Num() >= 1)
 	{
-		Initializer.NumNativeInstances = Initializer.NumNativeInstancesPerLayer[0];
+		Initializer.MaxNumInstances = Initializer.NumNativeInstancesPerLayer[0];
 	}
 	PRAGMA_ENABLE_DEPRECATION_WARNINGS
 	return GDynamicRHI->RHICreateRayTracingScene(MoveTemp(Initializer));
diff --git a/Engine/Source/Runtime/RHI/Public/RHIContext.h b/Engine/Source/Runtime/RHI/Public/RHIContext.h
--- a/Engine/Source/Runtime/RHI/Public/RHIContext.h
+++ b/Engine/Source/Runtime/RHI/Public/RHIContext.h
@@ -560,6 +560,8 @@ struct FRayTracingSceneBuildParams
 	FRHIBuffer* InstanceBuffer = nullptr;
 	uint32 InstanceBufferOffset = 0;
 
+	uint32 NumInstances = 0;
+
 	// Unique list of geometries referenced by all instances in this scene.
 	// Any referenced geometry is kept alive while the scene is alive.
 	TConstArrayView<FRHIRayTracingGeometry*> ReferencedGeometries;
diff --git a/Engine/Source/Runtime/RHI/Public/RHIResources.h b/Engine/Source/Runtime/RHI/Public/RHIResources.h
--- a/Engine/Source/Runtime/RHI/Public/RHIResources.h
+++ b/Engine/Source/Runtime/RHI/Public/RHIResources.h
@@ -3306,10 +3306,11 @@ struct FRayTracingSceneInitializer2
 	TArray<uint32> SegmentPrefixSum;
 
 	// Total flattened number of ray tracing geometry instances (a single FRayTracingGeometryInstance may represent many) per layer.
-	UE_DEPRECATED(5.5, "FRHIRayTracingScene layers are deprecated. Use NumNativeInstances and create one FRHIRayTracingScene per layer instead.")
+	UE_DEPRECATED(5.5, "FRHIRayTracingScene layers are deprecated. Use MaxNumInstances and create one FRHIRayTracingScene per layer instead.")
 	TArray<uint32> NumNativeInstancesPerLayer;
 
-	uint32 NumNativeInstances = 0;
+	// Maximum number of instances in this scene. Actual number of instances is specified in FRayTracingSceneBuildParams.
+	uint32 MaxNumInstances = 0;
 
 	uint32 NumTotalSegments = 0;
 
diff --git a/Engine/Source/Runtime/Renderer/Private/PathTracing.cpp b/Engine/Source/Runtime/Renderer/Private/PathTracing.cpp
--- a/Engine/Source/Runtime/Renderer/Private/PathTracing.cpp
+++ b/Engine/Source/Runtime/Renderer/Private/PathTracing.cpp
@@ -807,7 +807,7 @@ static void PreparePathTracingData(const FScene* Scene, const FViewInfo& View, F
 
 	PathTracingData.DecalRoughnessCutoff = PathTracing::UsesDecals(*View.Family) && View.bHasRayTracingDecals ? CVarPathTracingDecalRoughnessCutoff.GetValueOnRenderThread() : -1.0f;
 
-	PathTracingData.MeshDecalRoughnessCutoff = PathTracing::UsesDecals(*View.Family) && Scene->RayTracingScene.GetRHIRayTracingScene(ERayTracingSceneLayer::Decals)->GetInitializer().NumNativeInstances > 0 ? CVarPathTracingMeshDecalRoughnessCutoff.GetValueOnRenderThread() : -1.0f;
+	PathTracingData.MeshDecalRoughnessCutoff = PathTracing::UsesDecals(*View.Family) && Scene->RayTracingScene.GetNumNativeInstances(ERayTracingSceneLayer::Decals) > 0 ? CVarPathTracingMeshDecalRoughnessCutoff.GetValueOnRenderThread() : -1.0f;
 	PathTracingData.MeshDecalBias = CVarPathTracingMeshDecalBias.GetValueOnRenderThread();
 
 	PathTracingData.MaxRaymarchSteps = CVarPathTracingMaxRaymarchSteps.GetValueOnRenderThread();
diff --git a/Engine/Source/Runtime/Renderer/Private/RayTracing/RayTracingDebug.cpp b/Engine/Source/Runtime/Renderer/Private/RayTracing/RayTracingDebug.cpp
--- a/Engine/Source/Runtime/Renderer/Private/RayTracing/RayTracingDebug.cpp
+++ b/Engine/Source/Runtime/Renderer/Private/RayTracing/RayTracingDebug.cpp
@@ -1115,7 +1115,8 @@ static FRDGBufferRef RayTracingPerformHitStatsPerPrimitive(FRDGBuilder& GraphBui
 
 	FShaderBindingTableRHIRef HitStatsSBT = RHICreateShaderBindingTable(SBTInitializer);
 
-	const uint32 NumInstancesInTLAS = FMath::Max(RayTracingScene.GetRHIRayTracingSceneChecked(ERayTracingSceneLayer::Base)->GetInitializer().NumNativeInstances, (uint32)CVarRayTracingDebugHitCountTopKHits.GetValueOnRenderThread());
+	// TODO: Should check RayTracingScene for actual number of instances instead of max number in FRHIRayTracingScene initializer
+	const uint32 NumInstancesInTLAS = FMath::Max(RayTracingScene.GetRHIRayTracingSceneChecked(ERayTracingSceneLayer::Base)->GetInitializer().MaxNumInstances, (uint32)CVarRayTracingDebugHitCountTopKHits.GetValueOnRenderThread());
 	FRDGBufferDesc HitStatsPerPrimitiveBufferDesc = FRDGBufferDesc::CreateStructuredDesc(sizeof(FRayTracingHitStatsEntry), NumInstancesInTLAS);
 	HitStatsPerPrimitiveBufferDesc.Usage = EBufferUsageFlags(HitStatsPerPrimitiveBufferDesc.Usage | BUF_SourceCopy);
 	FRDGBufferRef HitStatsBuffer = GraphBuilder.CreateBuffer(HitStatsPerPrimitiveBufferDesc, TEXT("RayTracingDebug.HitStatsBuffer"));
diff --git a/Engine/Source/Runtime/Renderer/Private/RayTracing/RayTracingInstanceBufferUtil.cpp b/Engine/Source/Runtime/Renderer/Private/RayTracing/RayTracingInstanceBufferUtil.cpp
--- a/Engine/Source/Runtime/Renderer/Private/RayTracing/RayTracingInstanceBufferUtil.cpp
+++ b/Engine/Source/Runtime/Renderer/Private/RayTracing/RayTracingInstanceBufferUtil.cpp
@@ -198,7 +198,7 @@ FRayTracingSceneWithGeometryInstances CreateRayTracingSceneWithGeometryInstances
 
 	if(NumLayers == 1)
 	{
-		Initializer.NumNativeInstances = Initializer.NumNativeInstancesPerLayer[0];
+		Initializer.MaxNumInstances = Initializer.NumNativeInstancesPerLayer[0];
 		Initializer.NumNativeInstancesPerLayer.Empty();
 	}
 	PRAGMA_ENABLE_DEPRECATION_WARNINGS
diff --git a/Engine/Source/Runtime/Renderer/Private/RayTracing/RayTracingScene.cpp b/Engine/Source/Runtime/Renderer/Private/RayTracing/RayTracingScene.cpp
--- a/Engine/Source/Runtime/Renderer/Private/RayTracing/RayTracingScene.cpp
+++ b/Engine/Source/Runtime/Renderer/Private/RayTracing/RayTracingScene.cpp
@@ -97,7 +97,7 @@ void FRayTracingScene::Create(FRDGBuilder& GraphBuilder, const FViewInfo& View,
 
 			FRayTracingSceneInitializer2 Initializer;
 			Initializer.DebugName = FName(TEXT("FRayTracingScene"));
-			Initializer.NumNativeInstances = Layer.InitializationData.NumNativeGPUSceneInstances + Layer.InitializationData.NumNativeCPUInstances;
+			Initializer.MaxNumInstances = Layer.InitializationData.NumNativeGPUSceneInstances + Layer.InitializationData.NumNativeCPUInstances;
 			Initializer.NumTotalSegments = Layer.InitializationData.TotalNumSegments;
 			Initializer.BuildFlags = BuildFlags;
 
@@ -462,6 +462,7 @@ void FRayTracingScene::Build(FRDGBuilder& GraphBuilder, ERDGPassFlags ComputePas
 				BuildParams.ScratchBufferOffset = 0;
 				BuildParams.InstanceBuffer = PassParams->InstanceBuffer->GetRHI();
 				BuildParams.InstanceBufferOffset = 0;
+				BuildParams.NumInstances = Layer.InitializationData.NumNativeCPUInstances + Layer.InitializationData.NumNativeGPUSceneInstances;
 				BuildParams.ReferencedGeometries = Layer.InitializationData.ReferencedGeometries;
 				BuildParams.PerInstanceGeometries = Layer.InitializationData.PerInstanceGeometries;
 
diff --git a/Engine/Source/Runtime/Renderer/Private/Tests/RayTracingTestbed.cpp b/Engine/Source/Runtime/Renderer/Private/Tests/RayTracingTestbed.cpp
--- a/Engine/Source/Runtime/Renderer/Private/Tests/RayTracingTestbed.cpp
+++ b/Engine/Source/Runtime/Renderer/Private/Tests/RayTracingTestbed.cpp
@@ -184,7 +184,7 @@ bool RunRayTracingTestbed_RenderThread(const FString& Parameters)
 	{
 		FRayTracingSceneInitializer2 Initializer;
 		Initializer.DebugName = FName(TEXT("FRayTracingScene"));
-		Initializer.NumNativeInstances = RayTracingSceneInitializationData.NumNativeGPUSceneInstances + RayTracingSceneInitializationData.NumNativeCPUInstances;
+		Initializer.MaxNumInstances = RayTracingSceneInitializationData.NumNativeGPUSceneInstances + RayTracingSceneInitializationData.NumNativeCPUInstances;
 		Initializer.NumTotalSegments = RayTracingSceneInitializationData.TotalNumSegments;
 		Initializer.BuildFlags = SceneBuildFlags;
 
@@ -207,12 +207,12 @@ bool RunRayTracingTestbed_RenderThread(const FString& Parameters)
 		ScratchBufferCreateInfo);
 
 	FRWBufferStructured InstanceBuffer;
-	InstanceBuffer.Initialize(RHICmdList, TEXT("RayTracingTestBedInstanceBuffer"), GRHIRayTracingInstanceDescriptorSize, SceneInitializer.NumNativeInstances);
+	InstanceBuffer.Initialize(RHICmdList, TEXT("RayTracingTestBedInstanceBuffer"), GRHIRayTracingInstanceDescriptorSize, SceneInitializer.MaxNumInstances);
 
 	FByteAddressBuffer AccelerationStructureAddressesBuffer;
 	AccelerationStructureAddressesBuffer.Initialize(RHICmdList, TEXT("RayTracingTestBedAccelerationStructureAddressesBuffer"), sizeof(FRayTracingAccelerationStructureAddress), BUF_Volatile | BUF_MultiGPUAllocate);
 
-	const uint32 InstanceUploadBufferSize = SceneInitializer.NumNativeInstances * sizeof(FRayTracingInstanceDescriptorInput);
+	const uint32 InstanceUploadBufferSize = SceneInitializer.MaxNumInstances * sizeof(FRayTracingInstanceDescriptorInput);
 	FBufferRHIRef InstanceUploadBuffer;
 	FShaderResourceViewRHIRef InstanceUploadSRV;
 	{
@@ -242,7 +242,7 @@ bool RunRayTracingTestbed_RenderThread(const FString& Parameters)
 			RayTracingSceneInitializationData.BaseInstancePrefixSum,
 			RayTracingSceneInitializationData.NumNativeGPUSceneInstances,
 			RayTracingSceneInitializationData.NumNativeCPUInstances,
-			MakeArrayView(InstanceUploadData, SceneInitializer.NumNativeInstances),
+			MakeArrayView(InstanceUploadData, SceneInitializer.MaxNumInstances),
 			MakeArrayView(TransformUploadData, RayTracingSceneInitializationData.NumNativeCPUInstances * 3));
 		RHICmdList.UnlockBuffer(TransformUploadBuffer);
 		RHICmdList.UnlockBuffer(InstanceUploadBuffer);
@@ -298,6 +298,7 @@ bool RunRayTracingTestbed_RenderThread(const FString& Parameters)
 	BuildParams.InstanceBufferOffset = 0;
 	BuildParams.ReferencedGeometries = RayTracingSceneInitializationData.ReferencedGeometries;
 	BuildParams.PerInstanceGeometries = RayTracingSceneInitializationData.PerInstanceGeometries;
+	BuildParams.NumInstances = RayTracingSceneInitializationData.NumNativeGPUSceneInstances + RayTracingSceneInitializationData.NumNativeCPUInstances;
 
 	RHICmdList.BuildAccelerationStructure(BuildParams);
 
diff --git a/Engine/Source/Runtime/VulkanRHI/Private/VulkanRayTracing.cpp b/Engine/Source/Runtime/VulkanRHI/Private/VulkanRayTracing.cpp
--- a/Engine/Source/Runtime/VulkanRHI/Private/VulkanRayTracing.cpp
+++ b/Engine/Source/Runtime/VulkanRHI/Private/VulkanRayTracing.cpp
@@ -676,7 +676,7 @@ FVulkanRayTracingScene::FVulkanRayTracingScene(FRayTracingSceneInitializer2 InIn
 {
 	INC_DWORD_STAT(STAT_VulkanRayTracingAllocatedTLAS);
 
-	SizeInfo = RHICalcRayTracingSceneSize(Initializer.NumNativeInstances, Initializer.BuildFlags);
+	SizeInfo = RHICalcRayTracingSceneSize(Initializer.MaxNumInstances, Initializer.BuildFlags);
 
 	const uint32 ParameterBufferSize = FMath::Max<uint32>(1, Initializer.NumTotalSegments) * sizeof(FVulkanRayTracingGeometryParameters);
 	FRHIResourceCreateInfo ParameterBufferCreateInfo(TEXT("RayTracingSceneMetadata"));
@@ -736,29 +736,42 @@ void FVulkanRayTracingScene::BindBuffer(FRHIBuffer* InBuffer, uint32 InBufferOff
 	}
 }
 
-void FVulkanRayTracingScene::BuildAccelerationStructure(
+void BuildAccelerationStructure(
 	FVulkanCommandListContext& CommandContext,
+	FVulkanRayTracingScene& Scene,
 	FVulkanResourceMultiBuffer* InScratchBuffer, uint32 InScratchOffset,
 	FVulkanResourceMultiBuffer* InInstanceBuffer, uint32 InInstanceOffset,
+	uint32 NumInstances,
 	EAccelerationStructureBuildMode BuildMode)
 {
+	check(InInstanceBuffer != nullptr);
+	checkf(NumInstances <= Scene.Initializer.MaxNumInstances, TEXT("NumInstances must be less or equal to MaxNumInstances"));
+
+	checkf(Scene.AccelerationStructureBuffer.IsValid(), TEXT("A buffer must be bound to the ray tracing scene before it can be built."));
+	checkf(Scene.View.IsValid(), TEXT("A buffer must be bound to the ray tracing scene before it can be built."));
+
 	const bool bIsUpdate = BuildMode == EAccelerationStructureBuildMode::Update;
 
+	if (bIsUpdate)
+	{
+		checkf(NumInstances == Scene.NumInstances, TEXT("Number of instances used to update TLAS must match the number used to build."));
+	}
+	else
+	{
+		Scene.NumInstances = NumInstances;
+	}
+
 	FBufferRHIRef ScratchBuffer;
 	{
 		TRHICommandList_RecursiveHazardous<FVulkanCommandListContext> RHICmdList(&CommandContext);
 
 		// Build a metadata buffer	that contains VulkanRHI-specific per-geometry parameters that allow us to access
 		// vertex and index buffers from shaders that use inline ray tracing.
-		BuildPerInstanceGeometryParameterBuffer(RHICmdList);
-
-		check(AccelerationStructureBuffer.IsValid());
-		check(InInstanceBuffer != nullptr);
-
+		Scene.BuildPerInstanceGeometryParameterBuffer(RHICmdList);
 
 		if (InScratchBuffer == nullptr)
 		{
-			const uint64 ScratchBufferSize = bIsUpdate ? SizeInfo.UpdateScratchSize : SizeInfo.BuildScratchSize;
+			const uint64 ScratchBufferSize = bIsUpdate ? Scene.SizeInfo.UpdateScratchSize : Scene.SizeInfo.BuildScratchSize;
 
 			FRHIResourceCreateInfo ScratchBufferCreateInfo(TEXT("BuildScratchTLAS"));
 			ScratchBuffer = RHICmdList.CreateBuffer(ScratchBufferSize, BUF_StructuredBuffer | BUF_RayTracingScratch, 0, ERHIAccess::UAVCompute, ScratchBufferCreateInfo);
@@ -767,6 +780,15 @@ void FVulkanRayTracingScene::BuildAccelerationStructure(
 		}
 	}
 
+	if (bIsUpdate)
+	{
+		checkf(InScratchBuffer, TEXT("TLAS update requires scratch buffer of at least %lld bytes."), Scene.SizeInfo.UpdateScratchSize);
+	}
+	else
+	{
+		checkf(InScratchBuffer, TEXT("TLAS build requires scratch buffer of at least %lld bytes."), Scene.SizeInfo.BuildScratchSize);
+	}
+
 	FVkRtTLASBuildData BuildData;
 
 	VkAccelerationStructureBuildRangeInfoKHR BuildRangeInfo;
@@ -775,14 +797,13 @@ void FVulkanRayTracingScene::BuildAccelerationStructure(
 	const VkDeviceAddress InstanceBufferAddress = InInstanceBuffer->GetDeviceAddress() + InInstanceOffset;
 
 	{
-		GetTLASBuildData(Device->GetInstanceHandle(), Initializer.NumNativeInstances, InstanceBufferAddress, Initializer.BuildFlags, BuildMode, BuildData);
+		GetTLASBuildData(Scene.GetParent()->GetInstanceHandle(), NumInstances, InstanceBufferAddress, Scene.Initializer.BuildFlags, BuildMode, BuildData);
 
-		checkf(View.IsValid(), TEXT("A buffer must be bound to the ray tracing scene before it can be built."));
-		BuildData.GeometryInfo.dstAccelerationStructure = View->GetAccelerationStructureView().Handle;
-		BuildData.GeometryInfo.srcAccelerationStructure = bIsUpdate ? View->GetAccelerationStructureView().Handle : nullptr;
+		BuildData.GeometryInfo.dstAccelerationStructure = Scene.View->GetAccelerationStructureView().Handle;
+		BuildData.GeometryInfo.srcAccelerationStructure = bIsUpdate ? Scene.View->GetAccelerationStructureView().Handle : nullptr;
 		BuildData.GeometryInfo.scratchData.deviceAddress = InScratchBuffer->GetDeviceAddress() + InScratchOffset;
 
-		BuildRangeInfo.primitiveCount = Initializer.NumNativeInstances;
+		BuildRangeInfo.primitiveCount = NumInstances;
 		BuildRangeInfo.primitiveOffset = 0;
 		BuildRangeInfo.transformOffset = 0;
 		BuildRangeInfo.firstVertex = 0;
@@ -813,7 +834,7 @@ void FVulkanRayTracingScene::BuildAccelerationStructure(
 	CommandBufferManager.SubmitActiveCmdBuffer();
 	CommandBufferManager.PrepareForNewActiveCommandBuffer();
 
-	bBuilt = true;
+	Scene.bBuilt = true;
 }
 
 void FVulkanRayTracingScene::BuildPerInstanceGeometryParameterBuffer(FRHICommandListBase& RHICmdList)
@@ -1323,10 +1344,12 @@ void FVulkanCommandListContext::RHIBuildAccelerationStructure(const FRayTracingS
 
 	Scene->PerInstanceGeometries = SceneBuildParams.PerInstanceGeometries;
 
-	Scene->BuildAccelerationStructure(
+	BuildAccelerationStructure(
 		*this,
+		*Scene,
 		ScratchBuffer, SceneBuildParams.ScratchBufferOffset, 
 		InstanceBuffer, SceneBuildParams.InstanceBufferOffset,
+		SceneBuildParams.NumInstances,
 		SceneBuildParams.BuildMode);
 }
 
diff --git a/Engine/Source/Runtime/VulkanRHI/Private/VulkanRayTracing.h b/Engine/Source/Runtime/VulkanRHI/Private/VulkanRayTracing.h
--- a/Engine/Source/Runtime/VulkanRHI/Private/VulkanRayTracing.h
+++ b/Engine/Source/Runtime/VulkanRHI/Private/VulkanRayTracing.h
@@ -204,11 +204,6 @@ public:
 	const FRayTracingSceneInitializer2& GetInitializer() const override final { return Initializer; }
 
 	void BindBuffer(FRHIBuffer* InBuffer, uint32 InBufferOffset);
-	void BuildAccelerationStructure(
-		FVulkanCommandListContext& CommandContext, 
-		FVulkanResourceMultiBuffer* ScratchBuffer, uint32 ScratchOffset, 
-		FVulkanResourceMultiBuffer* InstanceBuffer, uint32 InstanceOffset,
-		EAccelerationStructureBuildMode BuildMode);
 
 	void CommitShaderTables(FVulkanCommandListContext& Context)
 	{
@@ -235,7 +230,10 @@ public:
 		return bBuilt;
 	}
 
-private:
+	void BuildPerInstanceGeometryParameterBuffer(FRHICommandListBase& RHICmdList);
+
+	using FRHIRayTracingAccelerationStructure::SizeInfo;
+
 	const FRayTracingSceneInitializer2 Initializer;
 
 	// Unique list of geometries referenced by all instances in this scene.
@@ -253,6 +251,8 @@ private:
 	// Many VkAccelerationStructureKHR-s may be created, pointing at the same buffer.
 
 	TUniquePtr<FVulkanView> View;
+
+	uint32 NumInstances = 0;
 	
 	TRefCountPtr<FVulkanResourceMultiBuffer> AccelerationStructureBuffer;
 
@@ -262,10 +262,10 @@ private:
 	
 	TMap<const FVulkanRayTracingPipelineState*, TRefCountPtr<FVulkanRayTracingShaderTable>> ShaderTables;
 
-	void BuildPerInstanceGeometryParameterBuffer(FRHICommandListBase& RHICmdList);
+	bool bBuilt = false;
 
+private:
 	UE::FMutex Mutex;
-	bool bBuilt = false;
 };
 
 
