Work on new RHI GPU profiler
 - Initial implementation of unified "stat unit" event consumer. Computes the union of busy time on the GPU across all registered queues, using the stream of events emitted by the platform RHI.
 - Also, initial implementation of ProfileGPU replacement, using the same stream of events. Currently triggered via "ProfileGPUNew" command. Does not yet include UI.

Removed adjusted timestamp support in D3D12 RHI
 - This has been disabled in D3D12 since it was added, as it breaks Insights markers.
 - Keeping this around complicates the D3D12 RHI implementation with the new profiler, so its not worth keeping.

Fix assert when running the editor with -onethread
 - AllowParallelTranslate() should return false if a lock fence was used on the command list (RHIThreadFence(true)). However, when running without an RHI thread, the LastLockFenceCommand field is never set.
 - Added bUsesLockFence boolean to track use of lock fences even when the RHI thread is disabled, so that checks for AllowParallelTranslate() inside platform RHIs don't fail erroneously.

diff --git a/Engine/Source/Runtime/Apple/MetalRHI/Private/MetalDynamicRHI.h b/Engine/Source/Runtime/Apple/MetalRHI/Private/MetalDynamicRHI.h
--- a/Engine/Source/Runtime/Apple/MetalRHI/Private/MetalDynamicRHI.h
+++ b/Engine/Source/Runtime/Apple/MetalRHI/Private/MetalDynamicRHI.h
@@ -42,7 +42,7 @@ public:
 	virtual ERHIInterfaceType GetInterfaceType() const override { return ERHIInterfaceType::Metal; }
 
 	virtual void RHIEndFrame_RenderThread(FRHICommandListImmediate& RHICmdList) final override;
-	virtual void RHIEndFrame() final override;
+	virtual void RHIEndFrame(const FRHIEndFrameArgs& Args) final override;
 
 	virtual FRHIShaderLibraryRef RHICreateShaderLibrary(EShaderPlatform Platform, FString const& FilePath, FString const& Name) final override;
 	virtual FSamplerStateRHIRef RHICreateSamplerState(const FSamplerStateInitializerRHI& Initializer) final override;
diff --git a/Engine/Source/Runtime/Apple/MetalRHI/Private/MetalRHI.cpp b/Engine/Source/Runtime/Apple/MetalRHI/Private/MetalRHI.cpp
--- a/Engine/Source/Runtime/Apple/MetalRHI/Private/MetalRHI.cpp
+++ b/Engine/Source/Runtime/Apple/MetalRHI/Private/MetalRHI.cpp
@@ -1266,7 +1266,7 @@ void FMetalDynamicRHI::RHIEndFrame_RenderThread(FRHICommandListImmediate& RHICmd
 	});
 }
 
-void FMetalDynamicRHI::RHIEndFrame()
+void FMetalDynamicRHI::RHIEndFrame(const FRHIEndFrameArgs& Args)
 {
 	// increment the internal frame counter
 	Device->IncrementFrameRHIThread();
diff --git a/Engine/Source/Runtime/D3D12RHI/Private/D3D12CommandContext.cpp b/Engine/Source/Runtime/D3D12RHI/Private/D3D12CommandContext.cpp
--- a/Engine/Source/Runtime/D3D12RHI/Private/D3D12CommandContext.cpp
+++ b/Engine/Source/Runtime/D3D12RHI/Private/D3D12CommandContext.cpp
@@ -200,11 +200,10 @@ void FD3D12ContextCommon::BindDiagnosticBuffer(FD3D12RootSignature const* RootSi
 
 	#if RHI_NEW_GPU_PROFILER
 		{
-			using namespace UE::RHI::GPUProfiler;
+			auto& Event = GetCommandList().EmplaceEvent<UE::RHI::GPUProfiler::FEvent::FBeginBreadcrumb>();
+			Event.Breadcrumb = Breadcrumb;
 
-			auto& Event = BreadcrumbEvents.Emplace_GetRef(MakeUnique<FBreadcrumbEvent>(FBreadcrumbEvent::FBeginBreadcrumb{ .Breadcrumb = Breadcrumb }));
-
-			FD3D12QueryLocation TimestampQuery = AllocateQuery(ED3D12QueryType::ProfilerTimestampTOP, &Event->Value.Get<FBreadcrumbEvent::FBeginBreadcrumb>().GPUTimestampTOP);
+			FD3D12QueryLocation TimestampQuery = AllocateQuery(ED3D12QueryType::ProfilerTimestampTOP, &Event.GPUTimestampTOP);
 			EndQuery(TimestampQuery);
 		}
 	#else
@@ -223,11 +222,10 @@ void FD3D12ContextCommon::BindDiagnosticBuffer(FD3D12RootSignature const* RootSi
 	{
 	#if RHI_NEW_GPU_PROFILER
 		{
-			using namespace UE::RHI::GPUProfiler;
-
-			auto& Event = BreadcrumbEvents.Emplace_GetRef(MakeUnique<FBreadcrumbEvent>(FBreadcrumbEvent::FEndBreadcrumb{ .Breadcrumb = Breadcrumb }));
+			auto& Event = GetCommandList().EmplaceEvent<UE::RHI::GPUProfiler::FEvent::FEndBreadcrumb>();
+			Event.Breadcrumb = Breadcrumb;
 
-			FD3D12QueryLocation TimestampQuery = AllocateQuery(ED3D12QueryType::ProfilerTimestampBOP, &Event->Value.Get<FBreadcrumbEvent::FEndBreadcrumb>().GPUTimestampBOP);
+			FD3D12QueryLocation TimestampQuery = AllocateQuery(ED3D12QueryType::ProfilerTimestampBOP, &Event.GPUTimestampBOP);
 			EndQuery(TimestampQuery);
 		}
 	#else
@@ -333,8 +331,8 @@ FD3D12QueryLocation FD3D12ContextCommon::AllocateQuery(ED3D12QueryType Type, voi
 		checkNoEntry();
 		[[fallthrough]];
 
-	case ED3D12QueryType::AdjustedRaw:
-	case ED3D12QueryType::AdjustedMicroseconds:
+	case ED3D12QueryType::TimestampRaw:
+	case ED3D12QueryType::TimestampMicroseconds:
 #if RHI_NEW_GPU_PROFILER
 	case ED3D12QueryType::ProfilerTimestampTOP:
 	case ED3D12QueryType::ProfilerTimestampBOP:
@@ -358,8 +356,8 @@ FD3D12QueryLocation FD3D12ContextCommon::InsertTimestamp(ED3D12Units Units, uint
 		checkNoEntry();
 		[[fallthrough]];
 
-	case ED3D12Units::Microseconds: Type = ED3D12QueryType::AdjustedMicroseconds; break;
-	case ED3D12Units::Raw:          Type = ED3D12QueryType::AdjustedRaw;          break;
+	case ED3D12Units::Microseconds: Type = ED3D12QueryType::TimestampMicroseconds; break;
+	case ED3D12Units::Raw:          Type = ED3D12QueryType::TimestampRaw;          break;
 	}
 
 	FD3D12QueryLocation Location = AllocateQuery(Type, Target);
@@ -427,10 +425,6 @@ void FD3D12ContextCommon::CloseCommandList()
 	TimestampQueries    .CloseAndReset(Payload->BatchedObjects.QueryRanges);
 	OcclusionQueries    .CloseAndReset(Payload->BatchedObjects.QueryRanges);
 	PipelineStatsQueries.CloseAndReset(Payload->BatchedObjects.QueryRanges);
-
-#if RHI_NEW_GPU_PROFILER
-	Payload->BatchedObjects.BreadcrumbEvents.Append(MoveTemp(BreadcrumbEvents));
-#endif
 }
 
 void FD3D12CommandContext::CloseCommandList()
diff --git a/Engine/Source/Runtime/D3D12RHI/Private/D3D12CommandContext.h b/Engine/Source/Runtime/D3D12RHI/Private/D3D12CommandContext.h
--- a/Engine/Source/Runtime/D3D12RHI/Private/D3D12CommandContext.h
+++ b/Engine/Source/Runtime/D3D12RHI/Private/D3D12CommandContext.h
@@ -286,6 +286,13 @@ private:
 	// A sync point signaled when all payloads in this context have completed.
 	FD3D12SyncPointRef ContextSyncPoint;
 
+public:
+	void BeginRecursiveCommand()
+	{
+		// Nothing to do
+	}
+
+protected:
 	// Returns the current command list (or creates a new one if the command list was not open).
 	FD3D12CommandList& GetCommandList()
 	{
@@ -294,13 +301,6 @@ private:
 		return *CommandList;
 	}
 
-public:
-	void BeginRecursiveCommand()
-	{
-		// Nothing to do
-	}
-
-protected:
 	enum class EMarkerType { In, Out };
 	void WriteMarker(D3D12_GPU_VIRTUAL_ADDRESS Address, uint32 Value, EMarkerType Type);
 
@@ -336,10 +336,6 @@ protected:
 		}
 	}
 
-#if RHI_NEW_GPU_PROFILER
-	TArray<TUniquePtr<UE::RHI::GPUProfiler::FBreadcrumbEvent>> BreadcrumbEvents; // @todo - new gpu profiler - optimize storage / allocation
-#endif
-
 public:
 	// Flushes any pending commands in this context to the GPU.
 	void FlushCommands(ED3D12FlushFlags FlushFlags = ED3D12FlushFlags::None);
diff --git a/Engine/Source/Runtime/D3D12RHI/Private/D3D12CommandList.cpp b/Engine/Source/Runtime/D3D12RHI/Private/D3D12CommandList.cpp
--- a/Engine/Source/Runtime/D3D12RHI/Private/D3D12CommandList.cpp
+++ b/Engine/Source/Runtime/D3D12RHI/Private/D3D12CommandList.cpp
@@ -322,6 +322,11 @@ void FD3D12CommandList::BeginLocalQueries()
 	{
 		if (State.BeginTimestamp)
 		{
+#if RHI_NEW_GPU_PROFILER
+			auto& Event = EmplaceEvent<UE::RHI::GPUProfiler::FEvent::FBeginWork>();
+			State.BeginTimestamp.Target = &Event.GPUTimestampTOP;
+#endif
+
 			EndQuery(State.BeginTimestamp);
 		}
 
@@ -345,6 +350,11 @@ void FD3D12CommandList::EndLocalQueries()
 
 		if (State.EndTimestamp)
 		{
+#if RHI_NEW_GPU_PROFILER
+			auto& Event = EmplaceEvent<UE::RHI::GPUProfiler::FEvent::FEndWork>();
+			State.EndTimestamp.Target = &Event.GPUTimestampBOP;
+#endif
+
 			EndQuery(State.EndTimestamp);
 		}
 
@@ -401,20 +411,22 @@ void FD3D12CommandList::EndQuery(FD3D12QueryLocation const& Location)
 				Position = ED3D12QueryPosition::BottomOfPipe;
 				break;
 
-			case ED3D12QueryType::CommandListBegin:
-			case ED3D12QueryType::IdleBegin:
 #if RHI_NEW_GPU_PROFILER
 			case ED3D12QueryType::ProfilerTimestampTOP:
+#else
+			case ED3D12QueryType::CommandListBegin:
+			case ED3D12QueryType::IdleBegin:
 #endif
 				Position = ED3D12QueryPosition::TopOfPipe;
 				break;
 
-			case ED3D12QueryType::CommandListEnd:
-			case ED3D12QueryType::IdleEnd:
-			case ED3D12QueryType::AdjustedMicroseconds:
-			case ED3D12QueryType::AdjustedRaw:
+			case ED3D12QueryType::TimestampMicroseconds:
+			case ED3D12QueryType::TimestampRaw:
 #if RHI_NEW_GPU_PROFILER
 			case ED3D12QueryType::ProfilerTimestampBOP:
+#else
+			case ED3D12QueryType::CommandListEnd:
+			case ED3D12QueryType::IdleEnd:
 #endif
 				Position = ED3D12QueryPosition::BottomOfPipe;
 				break;
@@ -422,9 +434,11 @@ void FD3D12CommandList::EndQuery(FD3D12QueryLocation const& Location)
 
 			WriteTimestamp(Location, Position);
 
+#if RHI_NEW_GPU_PROFILER == 0
 			// Command list begin/end timestamps are handled separately by the 
 			// submission thread, so shouldn't be in the TimestampQueries array.
 			if (Location.Type != ED3D12QueryType::CommandListBegin && Location.Type != ED3D12QueryType::CommandListEnd)
+#endif
 			{
 				State.TimestampQueries.Add(Location);
 			}
@@ -452,8 +466,13 @@ FD3D12CommandList::FState::FState(FD3D12CommandAllocator* CommandAllocator, FD3D
 
 	if (TimestampAllocator)
 	{
+#if RHI_NEW_GPU_PROFILER
+		BeginTimestamp = TimestampAllocator->Allocate(ED3D12QueryType::ProfilerTimestampTOP, nullptr);
+		EndTimestamp   = TimestampAllocator->Allocate(ED3D12QueryType::ProfilerTimestampBOP, nullptr);
+#else
 		BeginTimestamp = TimestampAllocator->Allocate(ED3D12QueryType::CommandListBegin, nullptr);
 		EndTimestamp   = TimestampAllocator->Allocate(ED3D12QueryType::CommandListEnd  , nullptr);
+#endif
 	}
 
 	if (PipelineStatsAllocator)
diff --git a/Engine/Source/Runtime/D3D12RHI/Private/D3D12CommandList.h b/Engine/Source/Runtime/D3D12RHI/Private/D3D12CommandList.h
--- a/Engine/Source/Runtime/D3D12RHI/Private/D3D12CommandList.h
+++ b/Engine/Source/Runtime/D3D12RHI/Private/D3D12CommandList.h
@@ -115,6 +115,19 @@ public:
 	FD3D12ResidencySet* CloseResidencySet();
 #endif // ENABLE_RESIDENCY_MANAGEMENT
 
+#if RHI_NEW_GPU_PROFILER
+	template <typename TEventType>
+	TEventType& EmplaceEvent()
+	{
+		return State.EmplaceEvent<TEventType>();
+	}
+
+	TArray<TUniquePtr<UE::RHI::GPUProfiler::FEvent>>&& RetrieveEvents()
+	{
+		return MoveTemp(State.Events);
+	}
+#endif // RHI_NEW_GPU_PROFILER
+
 private:
 	struct FInterfaces
 	{
@@ -299,5 +312,15 @@ private:
 		bool bLocalQueriesBegun = false;
 		bool bLocalQueriesEnded = false;
 
+#if RHI_NEW_GPU_PROFILER
+		TArray<TUniquePtr<UE::RHI::GPUProfiler::FEvent>> Events; // @todo dev-pr : optimize storage
+
+		template <typename TEventType>
+		TEventType& EmplaceEvent()
+		{
+			return Events.Emplace_GetRef(MakeUnique<UE::RHI::GPUProfiler::FEvent>(TEventType()))->Value.Get<TEventType>();
+		}
+#endif
+		
 	} State;
 };
diff --git a/Engine/Source/Runtime/D3D12RHI/Private/D3D12Device.cpp b/Engine/Source/Runtime/D3D12RHI/Private/D3D12Device.cpp
--- a/Engine/Source/Runtime/D3D12RHI/Private/D3D12Device.cpp
+++ b/Engine/Source/Runtime/D3D12RHI/Private/D3D12Device.cpp
@@ -74,6 +74,25 @@ FD3D12Queue::~FD3D12Queue()
 	check(PendingInterrupt.IsEmpty());
 }
 
+#if RHI_NEW_GPU_PROFILER
+UE::RHI::GPUProfiler::FQueue FD3D12Queue::GetProfilerQueue() const
+{
+	UE::RHI::GPUProfiler::FQueue Queue;
+	Queue.GPU = Device->GetGPUIndex();
+	Queue.Index = 0;
+
+	switch (QueueType)
+	{
+	default: checkNoEntry(); [[fallthrough]];
+	case ED3D12QueueType::Direct: Queue.Type = UE::RHI::GPUProfiler::FQueue::EType::Graphics; break;
+	case ED3D12QueueType::Async : Queue.Type = UE::RHI::GPUProfiler::FQueue::EType::Compute ; break;
+	case ED3D12QueueType::Copy  : Queue.Type = UE::RHI::GPUProfiler::FQueue::EType::Copy    ; break;
+	}
+
+	return Queue;
+}
+#endif // RHI_NEW_GPU_PROFILER
+
 FD3D12Device::FD3D12Device(FRHIGPUMask InGPUMask, FD3D12Adapter* InAdapter)
 	: FD3D12SingleNodeGPUObject(InGPUMask)
 	, FD3D12AdapterChild       (InAdapter)
diff --git a/Engine/Source/Runtime/D3D12RHI/Private/D3D12Device.h b/Engine/Source/Runtime/D3D12RHI/Private/D3D12Device.h
--- a/Engine/Source/Runtime/D3D12RHI/Private/D3D12Device.h
+++ b/Engine/Source/Runtime/D3D12RHI/Private/D3D12Device.h
@@ -66,12 +66,6 @@ class FD3D12Timing
 {
 public:
 	FD3D12Queue& Queue;
-
-	TArray<uint64> Timestamps;
-	int32 TimestampIndex = 0;
-
-	uint64 BusyCycles = 0;
-
 	D3D12_QUERY_DATA_PIPELINE_STATISTICS PipelineStats {};
 
 #if RHI_NEW_GPU_PROFILER
@@ -79,12 +73,17 @@ public:
 	uint64 GPUFrequency = 0, GPUTimestamp = 0;
 	uint64 CPUFrequency = 0, CPUTimestamp = 0;
 
-#if WITH_RHI_BREADCRUMBS
+	#if WITH_RHI_BREADCRUMBS
 	TArray<TSharedPtr<FRHIBreadcrumbAllocatorArray>> BreadcrumbAllocators {};
-	TArray<TUniquePtr<UE::RHI::GPUProfiler::FBreadcrumbEvent>> BreadcrumbEvents;
-#endif
+	#endif
 
-#endif
+	TArray<TUniquePtr<UE::RHI::GPUProfiler::FEvent>> Events;
+
+#else
+
+	TArray<uint64> Timestamps;
+	int32 TimestampIndex = 0;
+	uint64 BusyCycles = 0;
 
 	uint64 GetCurrentTimestamp()  const { return Timestamps[TimestampIndex]; }
 	uint64 GetPreviousTimestamp() const { return Timestamps[TimestampIndex - 1]; }
@@ -94,6 +93,8 @@ public:
 
 	void AdvanceTimestamp() { TimestampIndex++; }
 
+#endif
+
 	FD3D12Timing(FD3D12Queue& Queue)
 		: Queue(Queue)
 	{}
@@ -143,9 +144,6 @@ public:
 	// The active timing struct on this queue. Updated / accessed by the interrupt thread.
 	FD3D12Timing* Timing = nullptr;
 
-	uint64 CumulativeIdleTicks = 0;
-	uint64 LastEndTime = 0;
-
 	TUniquePtr<FD3D12DiagnosticBuffer> DiagnosticBuffer;
 
 	// On some hardware, some auxiliary queue types may not support tile mapping and a separate queue must be used
@@ -167,6 +165,10 @@ public:
 	FD3D12Queue(FD3D12Device* Device, ED3D12QueueType QueueType);
 	~FD3D12Queue();
 
+#if RHI_NEW_GPU_PROFILER
+	UE::RHI::GPUProfiler::FQueue GetProfilerQueue() const;
+#endif
+
 private:
 	// Internal fence which may be used before calling ExecuteCommandLists
 	FD3D12Fence ExecuteCommandListsFence;
diff --git a/Engine/Source/Runtime/D3D12RHI/Private/D3D12Query.h b/Engine/Source/Runtime/D3D12RHI/Private/D3D12Query.h
--- a/Engine/Source/Runtime/D3D12RHI/Private/D3D12Query.h
+++ b/Engine/Source/Runtime/D3D12RHI/Private/D3D12Query.h
@@ -14,16 +14,22 @@ struct FD3D12QueryLocation;
 enum class ED3D12QueryType
 {
 	None,
-	CommandListBegin,
-	CommandListEnd,
 	PipelineStats,
-	IdleBegin,
-	IdleEnd,
-	AdjustedMicroseconds,
-	AdjustedRaw,
+	TimestampMicroseconds,
+	TimestampRaw,
 #if RHI_NEW_GPU_PROFILER
+
+	// Top-of-pipe GPU timestamp, converted to the CPU clock domain (i.e. relative to FPlatformTime::Cycles64())
 	ProfilerTimestampTOP,
+
+	// Bottom-of-pipe GPU timestamp, converted to the CPU clock domain (i.e. relative to FPlatformTime::Cycles64())
 	ProfilerTimestampBOP,
+
+#else
+	CommandListBegin,
+	CommandListEnd,
+	IdleBegin,
+	IdleEnd,
 #endif
 	Occlusion
 };
diff --git a/Engine/Source/Runtime/D3D12RHI/Private/D3D12RHI.cpp b/Engine/Source/Runtime/D3D12RHI/Private/D3D12RHI.cpp
--- a/Engine/Source/Runtime/D3D12RHI/Private/D3D12RHI.cpp
+++ b/Engine/Source/Runtime/D3D12RHI/Private/D3D12RHI.cpp
@@ -413,34 +413,44 @@ void FD3D12DynamicRHI::EnqueueEndOfPipeTask(TUniqueFunction<void()> TaskFunc, TU
 	);
 }
 
-void FD3D12DynamicRHI::FlushTiming(bool bCreateNew)
+void FD3D12DynamicRHI::FlushTiming(bool bCreateNew, const FRHIEndFrameArgs& Args)
 {
+	bool bHasTiming = !CurrentTimingPerQueue.IsEmpty();
 	auto Lambda = [this, OldTiming = MoveTemp(CurrentTimingPerQueue)]()
 	{
-		if (OldTiming)
-		{
-			ProcessTimestamps(*OldTiming.Get());
-		}
+		ProcessTimestamps(OldTiming);
 	};
 
-	if (bCreateNew)
-	{
-		CurrentTimingPerQueue = MakeUnique<TIndirectArray<FD3D12Timing>>();
-		CurrentTimingPerQueue->Reserve(GD3D12MaxNumQueues);
-	}
-
-	EnqueueEndOfPipeTask(MoveTemp(Lambda), [&](FD3D12Payload& Payload)
+	EnqueueEndOfPipeTask(MoveTemp(Lambda), [&, Number = FrameNumber++](FD3D12Payload& Payload)
 	{
 		if (bCreateNew)
 		{
-			FD3D12Timing* NewTiming = new FD3D12Timing(Payload.Queue);
-			Payload.Timing = NewTiming;
-			CurrentTimingPerQueue->Add(NewTiming);
+			TUniquePtr<FD3D12Timing> NewTiming = MakeUnique<FD3D12Timing>(Payload.Queue);
+			Payload.Timing = NewTiming.Get();
+			CurrentTimingPerQueue.Emplace(MoveTemp(NewTiming));
 		}
 		else
 		{
 			Payload.Timing = nullptr;
 		}
+
+#if RHI_NEW_GPU_PROFILER
+		if (bHasTiming && Payload.Queue.QueueType != ED3D12QueueType::Copy)
+		{
+			UE::RHI::GPUProfiler::FEvent::FFrameBoundary Boundary;
+			Boundary.FrameNumber = Number;
+
+		#if WITH_RHI_BREADCRUMBS
+			switch (Payload.Queue.QueueType)
+			{
+			case ED3D12QueueType::Direct: Boundary.Breadcrumb = Args.GPUBreadcrumbs[ERHIPipeline::Graphics    ]; break;
+			case ED3D12QueueType::Async : Boundary.Breadcrumb = Args.GPUBreadcrumbs[ERHIPipeline::AsyncCompute]; break;
+			}
+		#endif
+
+			Payload.Events.Emplace(MakeUnique<UE::RHI::GPUProfiler::FEvent>(Boundary));
+		}
+#endif // RHI_NEW_GPU_PROFILER
 	});
 }
 
@@ -576,7 +586,7 @@ void FD3D12DynamicRHI::RHIEndFrame_RenderThread(FRHICommandListImmediate& RHICmd
 	});
 }
 
-void FD3D12DynamicRHI::RHIEndFrame()
+void FD3D12DynamicRHI::RHIEndFrame(const FRHIEndFrameArgs& Args)
 {
 	for (auto& Adapter : ChosenAdapters)
 	{
@@ -601,7 +611,7 @@ void FD3D12DynamicRHI::RHIEndFrame()
 	UpdateMemoryStats();
 
 	// Close the previous frame's timing and start a new one
-	FlushTiming(true);
+	FlushTiming(true, Args);
 
 	// Pump the interrupt queue to gather completed events
 	// (required if we're not using an interrupt thread).
diff --git a/Engine/Source/Runtime/D3D12RHI/Private/D3D12RHIPrivate.h b/Engine/Source/Runtime/D3D12RHI/Private/D3D12RHIPrivate.h
--- a/Engine/Source/Runtime/D3D12RHI/Private/D3D12RHIPrivate.h
+++ b/Engine/Source/Runtime/D3D12RHI/Private/D3D12RHIPrivate.h
@@ -235,9 +235,11 @@ public:
 	// Called when the GPU has crashed. This function will not return.
 	void ProcessInterruptQueueOnGPUCrash();
 
-	TUniquePtr<TIndirectArray<FD3D12Timing>> CurrentTimingPerQueue;
-	void FlushTiming(bool bCreateNew);
-	void ProcessTimestamps(TIndirectArray<FD3D12Timing>& TimingPerQueue);
+	using FD3D12TimingArray = TArray<TUniquePtr<FD3D12Timing>, TInlineAllocator<GD3D12MaxNumQueues>>;
+
+	FD3D12TimingArray CurrentTimingPerQueue;
+	void FlushTiming(bool bCreateNew, const FRHIEndFrameArgs& Args);
+	void ProcessTimestamps(FD3D12TimingArray const& TimingPerQueue);
 
 	void InitializeSubmissionPipe();
 	void ShutdownSubmissionPipe();
@@ -292,7 +294,7 @@ public:
 	}
 
 	virtual void RHIEndFrame_RenderThread(FRHICommandListImmediate& RHICmdList) final override;
-	virtual void RHIEndFrame() final override;
+	virtual void RHIEndFrame(const FRHIEndFrameArgs& Args) final override;
 
 	virtual FSamplerStateRHIRef RHICreateSamplerState(const FSamplerStateInitializerRHI& Initializer) final override;
 	virtual FRasterizerStateRHIRef RHICreateRasterizerState(const FRasterizerStateInitializerRHI& Initializer) final override;
@@ -551,6 +553,9 @@ protected:
 	TRefCountPtr<IDXGIFactory2> DXGIFactoryForDisplayList;
 #endif
 
+	// Counts the number of calls to RHIEndFrame, and is used in GPU profiler frame boundary events.
+	uint32 FrameNumber = 0;
+
 public:
 
 	virtual FD3D12ResourceDesc GetResourceDesc(const FRHITextureDesc& CreateInfo) const;
diff --git a/Engine/Source/Runtime/D3D12RHI/Private/D3D12Resources.cpp b/Engine/Source/Runtime/D3D12RHI/Private/D3D12Resources.cpp
--- a/Engine/Source/Runtime/D3D12RHI/Private/D3D12Resources.cpp
+++ b/Engine/Source/Runtime/D3D12RHI/Private/D3D12Resources.cpp
@@ -1661,9 +1661,26 @@ void FD3D12ResourceBarrierBatcher::AddAliasingBarrier(ID3D12Resource* InResource
 
 void FD3D12ResourceBarrierBatcher::FlushIntoCommandList(FD3D12CommandList& CommandList, FD3D12QueryAllocator& TimestampAllocator)
 {
-	auto InsertTimestamp = [&](ED3D12QueryType Type)
+	auto InsertTimestamp = [&](bool bBegin)
 	{
+#if RHI_NEW_GPU_PROFILER
+		if (bBegin)
+		{
+			auto& Event = CommandList.EmplaceEvent<UE::RHI::GPUProfiler::FEvent::FEndWork>();
+			CommandList.EndQuery(TimestampAllocator.Allocate(ED3D12QueryType::ProfilerTimestampBOP, &Event.GPUTimestampBOP));
+		}
+		else
+		{
+			auto& Event = CommandList.EmplaceEvent<UE::RHI::GPUProfiler::FEvent::FBeginWork>();
+			CommandList.EndQuery(TimestampAllocator.Allocate(ED3D12QueryType::ProfilerTimestampTOP, &Event.GPUTimestampTOP));
+		}
+#else
+		ED3D12QueryType Type = bBegin
+			? ED3D12QueryType::IdleBegin
+			: ED3D12QueryType::IdleEnd;
+
 		CommandList.EndQuery(TimestampAllocator.Allocate(Type, nullptr));
+#endif
 	};
 
 	for (int32 BatchStart = 0, BatchEnd = 0; BatchStart < Barriers.Num(); BatchStart = BatchEnd)
@@ -1681,7 +1698,7 @@ void FD3D12ResourceBarrierBatcher::FlushIntoCommandList(FD3D12CommandList& Comma
 		// Insert an idle begin/end timestamp around the barrier batch if required.
 		if (bIdle)
 		{
-			InsertTimestamp(ED3D12QueryType::IdleBegin);
+			InsertTimestamp(true);
 		}
 
 		CommandList.GraphicsCommandList()->ResourceBarrier(BatchEnd - BatchStart, &Barriers[BatchStart]);
@@ -1695,7 +1712,7 @@ void FD3D12ResourceBarrierBatcher::FlushIntoCommandList(FD3D12CommandList& Comma
 
 		if (bIdle)
 		{
-			InsertTimestamp(ED3D12QueryType::IdleEnd);
+			InsertTimestamp(false);
 		}
 	}
 
diff --git a/Engine/Source/Runtime/D3D12RHI/Private/D3D12Submission.cpp b/Engine/Source/Runtime/D3D12RHI/Private/D3D12Submission.cpp
--- a/Engine/Source/Runtime/D3D12RHI/Private/D3D12Submission.cpp
+++ b/Engine/Source/Runtime/D3D12RHI/Private/D3D12Submission.cpp
@@ -14,9 +14,6 @@
 #define D3D12_USE_SUBMISSION_THREAD (1)
 #define D3D12_USE_INTERRUPT_THREAD  (1 && D3D12_PLATFORM_SUPPORTS_BLOCKING_FENCES)
 
-// When enabled, GPU timestamp queries are adjusted to remove idle time caused by CPU bubbles.
-#define D3D12_ENABLE_ADJUSTED_TIMESTAMPS 0		// Adjusted timestamps break Unreal Insights, disabling for now
-
 static TAutoConsoleVariable<int32> CVarRHIUseSubmissionThread(
 	TEXT("rhi.UseSubmissionThread"),
 	2,
@@ -122,6 +119,25 @@ private:
 
 void FD3D12DynamicRHI::InitializeSubmissionPipe()
 {
+#if RHI_NEW_GPU_PROFILER
+	{
+		TArray<UE::RHI::GPUProfiler::FQueue> Queues;
+		ForEachQueue([&](FD3D12Queue& Queue)
+		{
+			if (Queue.QueueType == ED3D12QueueType::Copy)
+			{
+				// @todo dev-pr : Skip copy queues for now. We don't have copy queue support in RHI contexts,
+				// so it's not possible to insert frame boundary events in a way that is pipelined correctly.
+				return;
+			}
+
+			Queues.Add(Queue.GetProfilerQueue());
+		});
+
+		UE::RHI::GPUProfiler::InitializeQueues(Queues);
+	}
+#endif
+
 	if (FPlatformProcess::SupportsMultithreading())
 	{
 #if D3D12_USE_INTERRUPT_THREAD
@@ -146,12 +162,32 @@ void FD3D12DynamicRHI::InitializeSubmissionPipe()
 #endif
 	}
 
-	FlushTiming(true);
+#if WITH_RHI_BREADCRUMBS
+	TRHIPipelineArray<FRHIBreadcrumbNode*> GPUBreadcrumbs { InPlace, nullptr };
+	FRHIEndFrameArgs Args
+	{
+		.GPUBreadcrumbs = GPUBreadcrumbs
+	};
+#else
+	FRHIEndFrameArgs Args;
+#endif
+	
+	FlushTiming(true, Args);
 }
 
 void FD3D12DynamicRHI::ShutdownSubmissionPipe()
 {
-	FlushTiming(false);
+#if WITH_RHI_BREADCRUMBS
+	TRHIPipelineArray<FRHIBreadcrumbNode*> GPUBreadcrumbs { InPlace, nullptr };
+	FRHIEndFrameArgs Args
+	{
+		.GPUBreadcrumbs = GPUBreadcrumbs
+	};
+#else
+	FRHIEndFrameArgs Args;
+#endif
+
+	FlushTiming(false, Args);
 
 	delete SubmissionThread;
 	SubmissionThread = nullptr;
@@ -425,6 +461,9 @@ FD3D12DynamicRHI::FProcessResult FD3D12DynamicRHI::ProcessSubmissionQueue()
 						TargetQueue.BatchedObjects.OcclusionQueries.Append(MoveTemp(CommandList->State.OcclusionQueries));
 						TargetQueue.BatchedObjects.PipelineStatsQueries.Append(MoveTemp(CommandList->State.PipelineStatsQueries));
 
+#if RHI_NEW_GPU_PROFILER
+						TargetQueue.BatchedObjects.TimestampQueries.Append(MoveTemp(CommandList->State.TimestampQueries));
+#else
 						// Timestamp Queries
 						if (CommandList->State.BeginTimestamp)
 						{
@@ -455,6 +494,7 @@ FD3D12DynamicRHI::FProcessResult FD3D12DynamicRHI::ProcessSubmissionQueue()
 							check(CommandList->State.TimestampQueries.IsEmpty());
 							check(!CommandList->State.EndTimestamp);
 						}
+#endif
 					};
 
 					for (int32 Index = 0; Index < Payload->CommandListsToExecute.Num(); Index++)
@@ -698,8 +738,6 @@ uint64 FD3D12Queue::FinalizePayload(bool bRequiresSignal, FPayloadArray& Payload
 	{
 		BatchedObjects.BreadcrumbAllocators.Append(MoveTemp(PayloadToSubmit->BatchedObjects.BreadcrumbAllocators));
 	}
-
-	BatchedObjects.BreadcrumbEvents.Append(MoveTemp(PayloadToSubmit->BatchedObjects.BreadcrumbEvents));
 #endif
 
 	// Gather query ranges from this payload, grouping by heap pointer
@@ -902,6 +940,8 @@ void FD3D12DynamicRHI::FlushBatchedPayloads(FD3D12Queue::FPayloadArray& Payloads
 		TArray<FD3D12ResidencySet*, TInlineAllocator<128>> ResidencySets;
 #endif
 
+		uint64 Time = FPlatformTime::Cycles64();
+
 		// Accumulate the command lists from the payload
 		for (uint32 Index = FirstPayload; Index < LastPayload; ++Index)
 		{
@@ -911,6 +951,18 @@ void FD3D12DynamicRHI::FlushBatchedPayloads(FD3D12Queue::FPayloadArray& Payloads
 			for (FD3D12CommandList* CommandList : Payload->CommandListsToExecute)
 			{
 				check(CommandList->IsClosed());
+
+#if RHI_NEW_GPU_PROFILER
+				{
+					TArray<TUniquePtr<UE::RHI::GPUProfiler::FEvent>> Events = CommandList->RetrieveEvents();
+					if (Events.Num() > 0)
+					{
+						Events[0]->Value.Get<UE::RHI::GPUProfiler::FEvent::FBeginWork>().CPUTimestamp = Time;
+						Payload->Events.Append(MoveTemp(Events));
+					}
+				}
+#endif // RHI_NEW_GPU_PROFILER
+
 				D3DCommandLists.Add(CommandList->Interfaces.CommandList);
 
 #if ENABLE_RESIDENCY_MANAGEMENT
@@ -1004,12 +1056,13 @@ void FD3D12DynamicRHI::FlushBatchedPayloads(FD3D12Queue::FPayloadArray& Payloads
 		{
 			Flush();
 
+			if (FD3D12Timing* LocalTiming = *Payload->Timing)
 			{
 				SCOPED_NAMED_EVENT(CalibrateClocks, FColor::Red);
-				FD3D12Timing* LocalTiming = *Payload->Timing;
+
 				// Calibrate the GPU timestamp / clock
-				VERIFYD3D12RESULT(D3DCommandQueue->GetClockCalibration(&LocalTiming->GPUTimestamp, &LocalTiming->CPUTimestamp));
-				VERIFYD3D12RESULT(D3DCommandQueue->GetTimestampFrequency(&LocalTiming->GPUFrequency));
+				VERIFYD3D12RESULT(Payload->Queue.D3DCommandQueue->GetClockCalibration(&LocalTiming->GPUTimestamp, &LocalTiming->CPUTimestamp));
+				VERIFYD3D12RESULT(Payload->Queue.D3DCommandQueue->GetTimestampFrequency(&LocalTiming->GPUFrequency));
 				QueryPerformanceFrequency(reinterpret_cast<LARGE_INTEGER*>(&LocalTiming->CPUFrequency));
 			}
 		}
@@ -1282,12 +1335,6 @@ FD3D12DynamicRHI::FProcessResult FD3D12DynamicRHI::ProcessInterruptQueue()
 
 			// Resolve query results
 			{
-				if (Payload->Timing.IsSet())
-				{
-					// Switch the new timing struct into the queue. This redirects timestamp results to separate each frame's work.
-					CurrentQueue.Timing = Payload->Timing.GetValue();
-				}
-
 				for (FD3D12QueryLocation& Query : Payload->BatchedObjects.OcclusionQueries)
 				{
 					check(Query.Target);
@@ -1309,9 +1356,6 @@ FD3D12DynamicRHI::FProcessResult FD3D12DynamicRHI::ProcessInterruptQueue()
 
 				if (Payload->BatchedObjects.TimestampQueries.Num())
 				{
-					FD3D12QueryLocation* IdleBegin = nullptr;
-					FD3D12QueryLocation* ListBegin = nullptr;
-
 					// Some timestamp queries report in microseconds
 					const double MicrosecondsScale = 1000000.0 / CurrentQueue.Device->GetTimestampFrequency(CurrentQueue.QueueType);
 
@@ -1324,50 +1368,10 @@ FD3D12DynamicRHI::FProcessResult FD3D12DynamicRHI::ProcessInterruptQueue()
 
 						switch (Query.Type)
 						{
-						case ED3D12QueryType::CommandListBegin:
-						case ED3D12QueryType::CommandListEnd:
-						case ED3D12QueryType::IdleBegin:
-						case ED3D12QueryType::IdleEnd:
-							check(CurrentQueue.Timing);
-							CurrentQueue.Timing->Timestamps.Add(Query.GetResult<uint64>());
-							break;
-						}
-
-						switch (Query.Type)
-						{
-						case ED3D12QueryType::CommandListBegin:
-							check(!ListBegin && !IdleBegin);
-							ListBegin = &Query;
-							break;
-
-						case ED3D12QueryType::CommandListEnd:
-							check(ListBegin != nullptr && IdleBegin == nullptr);
-							// Accumulate the number of ticks that have elapsed between the end of the previous command list, and the start of this one.
-							CurrentQueue.CumulativeIdleTicks += (CurrentQueue.LastEndTime != 0) ? ListBegin->GetResult<uint64>() - CurrentQueue.LastEndTime : 0; //-V522
-							CurrentQueue.LastEndTime = Query.GetResult<uint64>();
-							ListBegin = nullptr;
-							break;
-
-						case ED3D12QueryType::IdleBegin:
-							check(ListBegin && !IdleBegin);
-							IdleBegin = &Query;
-							break;
-
-						case ED3D12QueryType::IdleEnd:
-							check(ListBegin != nullptr && IdleBegin != nullptr);
-							// Accumulate the time this pipe spent in an idle scope. This includes vsync and waiting on other pipes.
-							CurrentQueue.CumulativeIdleTicks += Query.GetResult<uint64>() - IdleBegin->GetResult<uint64>(); //-V522
-							IdleBegin = nullptr;
-							break;
-
-						case ED3D12QueryType::AdjustedMicroseconds:
-						case ED3D12QueryType::AdjustedRaw:
-							check(ListBegin && !IdleBegin && Query.Target);
-#if D3D12_ENABLE_ADJUSTED_TIMESTAMPS
-							// Adjust the time such that the ticks reported only advance when this pipe is busy
-							*Query.Target -= CurrentQueue.CumulativeIdleTicks;
-#endif
-							if (Query.Type == ED3D12QueryType::AdjustedMicroseconds)
+						case ED3D12QueryType::TimestampMicroseconds:
+						case ED3D12QueryType::TimestampRaw:
+							check(Query.Target);
+							if (Query.Type == ED3D12QueryType::TimestampMicroseconds)
 							{
 								// Convert to microseconds
 								*static_cast<uint64*>(Query.Target) = FPlatformMath::TruncToInt(double(*static_cast<uint64*>(Query.Target)) * MicrosecondsScale);
@@ -1387,20 +1391,18 @@ FD3D12DynamicRHI::FProcessResult FD3D12DynamicRHI::ProcessInterruptQueue()
 								Target = CPUDelta + CurrentQueue.Timing->CPUTimestamp;
 							}
 							break;
-#endif // RHI_NEW_GPU_PROFILER
+#else
+						case ED3D12QueryType::CommandListBegin:
+						case ED3D12QueryType::CommandListEnd:
+						case ED3D12QueryType::IdleBegin:
+						case ED3D12QueryType::IdleEnd:
+							check(CurrentQueue.Timing);
+							CurrentQueue.Timing->Timestamps.Add(Query.GetResult<uint64>());
+							break;
+
+#endif
 						}
 					}
-
-					check(!ListBegin && !IdleBegin);
-				}
-			}
-
-			// Signal the CPU events of all sync points associated with this batch.
-			for (FD3D12SyncPointRef& SyncPoint : Payload->SyncPointsToSignal)
-			{
-				if (SyncPoint->GraphEvent)
-				{
-					SyncPoint->GraphEvent->DispatchSubsequents();
 				}
 			}
 
@@ -1411,13 +1413,28 @@ FD3D12DynamicRHI::FProcessResult FD3D12DynamicRHI::ProcessInterruptQueue()
 				CurrentQueue.Timing->BreadcrumbAllocators.Append(MoveTemp(Payload->BatchedObjects.BreadcrumbAllocators));
 			}
 
-			if (Payload->BatchedObjects.BreadcrumbEvents.Num())
+			if (Payload->Events.Num())
 			{
 				check(CurrentQueue.Timing);
-				CurrentQueue.Timing->BreadcrumbEvents.Append(MoveTemp(Payload->BatchedObjects.BreadcrumbEvents));
+				CurrentQueue.Timing->Events.Append(MoveTemp(Payload->Events));
 			}
 #endif
 
+			if (Payload->Timing.IsSet())
+			{
+				// Switch the new timing struct into the queue. This redirects timestamp results to separate each frame's work.
+				CurrentQueue.Timing = Payload->Timing.GetValue();
+			}
+
+			// Signal the CPU events of all sync points associated with this batch.
+			for (FD3D12SyncPointRef& SyncPoint : Payload->SyncPointsToSignal)
+			{
+				if (SyncPoint->GraphEvent)
+				{
+					SyncPoint->GraphEvent->DispatchSubsequents();
+				}
+			}
+
 			// We're done with this payload now.
 
 			// GPU resources the payload is holding a reference to will be cleaned up here.
@@ -1461,26 +1478,16 @@ static TAutoConsoleVariable<int32> CVarGPUTimeFromTimestamps(
 	TEXT("Prefer timestamps instead of GetHardwareGPUFrameTime to compute GPU frame time"),
 	ECVF_RenderThreadSafe);
 
-void FD3D12DynamicRHI::ProcessTimestamps(TIndirectArray<FD3D12Timing>& TimingPerQueue)
+void FD3D12DynamicRHI::ProcessTimestamps(FD3D12TimingArray const& TimingPerQueue)
 {
 #if RHI_NEW_GPU_PROFILER
-	for (FD3D12Timing const& Timing : TimingPerQueue)
-	{
-		UE::RHI::GPUProfiler::FQueue QueueID {};
-		QueueID.GPU = Timing.Queue.Device->GetGPUIndex();
-		QueueID.Index = 0;
 
-		switch (Timing.Queue.QueueType)
-		{
-		default: checkNoEntry(); [[fallthrough]];
-		case ED3D12QueueType::Direct: QueueID.Type = UE::RHI::GPUProfiler::FQueue::EType::Graphics; break;
-		case ED3D12QueueType::Async : QueueID.Type = UE::RHI::GPUProfiler::FQueue::EType::Compute ; break;
-		case ED3D12QueueType::Copy  : QueueID.Type = UE::RHI::GPUProfiler::FQueue::EType::Copy    ; break;
-		}
-		
-		UE::RHI::GPUProfiler::PushEvents(QueueID, Timing.BreadcrumbEvents, {}, {});
+	for (auto const& Timing : TimingPerQueue)
+	{
+		UE::RHI::GPUProfiler::ProcessEvents(Timing->Queue.GetProfilerQueue(), Timing->Events);
 	}
-#endif // RHI_NEW_GPU_PROFILER
+
+#else
 
 	// The total number of cycles where at least one GPU pipe was busy during the frame.
 	uint64 UnionBusyCycles = 0;
@@ -1494,11 +1501,11 @@ void FD3D12DynamicRHI::ProcessTimestamps(TIndirectArray<FD3D12Timing>& TimingPer
 	{
 		// Find the next minimum timestamp
 		FD3D12Timing* NextMin = nullptr;
-		for (FD3D12Timing& Current : TimingPerQueue)
+		for (auto const& Current : TimingPerQueue)
 		{
-			if (Current.HasMoreTimestamps() && (!NextMin || Current.GetCurrentTimestamp() < NextMin->GetCurrentTimestamp()))
+			if (Current->HasMoreTimestamps() && (!NextMin || Current->GetCurrentTimestamp() < NextMin->GetCurrentTimestamp()))
 			{
-				NextMin = &Current;
+				NextMin = Current.Get();
 			}
 		}
 
@@ -1531,10 +1538,12 @@ void FD3D12DynamicRHI::ProcessTimestamps(TIndirectArray<FD3D12Timing>& TimingPer
 
 	check(BusyPipes == 0);
 	
+#endif
+
 	D3D12_QUERY_DATA_PIPELINE_STATISTICS PipelineStats{};
-	for (FD3D12Timing& Current : TimingPerQueue)
+	for (auto const& Current : TimingPerQueue)
 	{
-		PipelineStats += Current.PipelineStats;
+		PipelineStats += Current->PipelineStats;
 	}
 
 	SET_DWORD_STAT(STAT_D3D12RHI_IAVertices   , PipelineStats.IAVertices   );
@@ -1549,6 +1558,8 @@ void FD3D12DynamicRHI::ProcessTimestamps(TIndirectArray<FD3D12Timing>& TimingPer
 	SET_DWORD_STAT(STAT_D3D12RHI_DSInvocations, PipelineStats.DSInvocations);
 	SET_DWORD_STAT(STAT_D3D12RHI_CSInvocations, PipelineStats.CSInvocations);
 
+#if RHI_NEW_GPU_PROFILER == 0
+
 	// @todo mgpu - how to handle multiple devices / queues with potentially different timestamp frequencies?
 	FD3D12Device* Device = GetAdapter().GetDevice(0);
 	double Frequency = Device->GetTimestampFrequency(ED3D12QueueType::Direct);
@@ -1571,13 +1582,15 @@ void FD3D12DynamicRHI::ProcessTimestamps(TIndirectArray<FD3D12Timing>& TimingPer
 		GGPUFrameTime = FPlatformMath::TruncToInt(double(UnionBusyCycles) * Scale32);
 	}
 
-	for (FD3D12Timing& Current : TimingPerQueue)
+	for (auto const& Current : TimingPerQueue)
 	{
-		switch (Current.Queue.QueueType)
+		switch (Current->Queue.QueueType)
 		{
-		case ED3D12QueueType::Direct: SET_CYCLE_COUNTER(STAT_RHI_GPUTotalTimeGraphics    , FPlatformMath::TruncToInt(double(Current.BusyCycles) * Scale64)); break;
-		case ED3D12QueueType::Async : SET_CYCLE_COUNTER(STAT_RHI_GPUTotalTimeAsyncCompute, FPlatformMath::TruncToInt(double(Current.BusyCycles) * Scale64)); break;
-		case ED3D12QueueType::Copy  : SET_CYCLE_COUNTER(STAT_RHI_GPUTotalTimeCopy        , FPlatformMath::TruncToInt(double(Current.BusyCycles) * Scale64)); break;
+		case ED3D12QueueType::Direct: SET_CYCLE_COUNTER(STAT_RHI_GPUTotalTimeGraphics    , FPlatformMath::TruncToInt(double(Current->BusyCycles) * Scale64)); break;
+		case ED3D12QueueType::Async : SET_CYCLE_COUNTER(STAT_RHI_GPUTotalTimeAsyncCompute, FPlatformMath::TruncToInt(double(Current->BusyCycles) * Scale64)); break;
+		case ED3D12QueueType::Copy  : SET_CYCLE_COUNTER(STAT_RHI_GPUTotalTimeCopy        , FPlatformMath::TruncToInt(double(Current->BusyCycles) * Scale64)); break;
 		}
 	}
+
+#endif
 }
diff --git a/Engine/Source/Runtime/D3D12RHI/Private/D3D12Submission.h b/Engine/Source/Runtime/D3D12RHI/Private/D3D12Submission.h
--- a/Engine/Source/Runtime/D3D12RHI/Private/D3D12Submission.h
+++ b/Engine/Source/Runtime/D3D12RHI/Private/D3D12Submission.h
@@ -150,7 +150,6 @@ struct FD3D12BatchedPayloadObjects
 
 #if WITH_RHI_BREADCRUMBS && RHI_NEW_GPU_PROFILER
 	TArray<TSharedPtr<FRHIBreadcrumbAllocatorArray>, TInlineAllocator<1>> BreadcrumbAllocators {};
-	TArray<TUniquePtr<UE::RHI::GPUProfiler::FBreadcrumbEvent>> BreadcrumbEvents;
 #endif
 
 	bool IsEmpty() const
@@ -162,7 +161,6 @@ struct FD3D12BatchedPayloadObjects
 			&& QueryRanges         .Num() == 0
 #if WITH_RHI_BREADCRUMBS && RHI_NEW_GPU_PROFILER
 			&& BreadcrumbAllocators.Num() == 0
-			&& BreadcrumbEvents    .Num() == 0
 #endif
 		;
 
@@ -247,6 +245,9 @@ struct FD3D12PayloadBase : public FD3D12PayloadBaseFixLayout
 #if WITH_RHI_BREADCRUMBS
 	FRHIBreadcrumbRange BreadcrumbRange {};
 #endif
+#if RHI_NEW_GPU_PROFILER
+	TArray<TUniquePtr<UE::RHI::GPUProfiler::FEvent>> Events;
+#endif
 
 	virtual ~FD3D12PayloadBase();
 
diff --git a/Engine/Source/Runtime/Engine/Private/UnrealEngine.cpp b/Engine/Source/Runtime/Engine/Private/UnrealEngine.cpp
--- a/Engine/Source/Runtime/Engine/Private/UnrealEngine.cpp
+++ b/Engine/Source/Runtime/Engine/Private/UnrealEngine.cpp
@@ -427,9 +427,6 @@ static TAutoConsoleVariable<float> CVarForceDynamicResScreenPercentage(
 	ECVF_Default | ECVF_RenderThreadSafe);
 #endif
 
-
-ENGINE_API uint32 GGPUFrameTime = 0;
-
 /** System resolution instance */
 FSystemResolution GSystemResolution;
 
diff --git a/Engine/Source/Runtime/Engine/Public/EngineGlobals.h b/Engine/Source/Runtime/Engine/Public/EngineGlobals.h
--- a/Engine/Source/Runtime/Engine/Public/EngineGlobals.h
+++ b/Engine/Source/Runtime/Engine/Public/EngineGlobals.h
@@ -9,5 +9,3 @@
 /** when set, disallows all network travel (pending level rejects all client travel attempts) */
 extern ENGINE_API bool						GDisallowNetworkTravel;
 
-/** The GPU time taken to render the last frame. Same metric as FPlatformTime::Cycles(). */
-extern ENGINE_API uint32					GGPUFrameTime;
diff --git a/Engine/Source/Runtime/NullDrv/Public/NullRHI.h b/Engine/Source/Runtime/NullDrv/Public/NullRHI.h
--- a/Engine/Source/Runtime/NullDrv/Public/NullRHI.h
+++ b/Engine/Source/Runtime/NullDrv/Public/NullRHI.h
@@ -340,7 +340,7 @@ public:
 		return new FNullTexture(Desc);
 	}
 
-	virtual void RHIEndFrame() final override
+	virtual void RHIEndFrame(const FRHIEndFrameArgs& Args) final override
 	{
 
 	}
diff --git a/Engine/Source/Runtime/OpenGLDrv/Private/OpenGLDevice.cpp b/Engine/Source/Runtime/OpenGLDrv/Private/OpenGLDevice.cpp
--- a/Engine/Source/Runtime/OpenGLDrv/Private/OpenGLDevice.cpp
+++ b/Engine/Source/Runtime/OpenGLDrv/Private/OpenGLDevice.cpp
@@ -217,7 +217,7 @@ FOpenGLContextState& FOpenGLDynamicRHI::GetContextStateForCurrentContext()
 	return *CachedContextState;
 }
 
-void FOpenGLDynamicRHI::RHIEndFrame()
+void FOpenGLDynamicRHI::RHIEndFrame(const FRHIEndFrameArgs& Args)
 {
 #if (RHI_NEW_GPU_PROFILER == 0)
 	GPUProfilingData->EndFrame();
@@ -1369,7 +1369,9 @@ FOpenGLDynamicRHI::FOpenGLDynamicRHI()
 	PrivateOpenGLDevicePtr = this;
 	GlobalUniformBuffers.AddZeroed(FUniformBufferStaticSlotRegistry::Get().GetSlotCount());
 
+#if RHI_NEW_GPU_PROFILER == 0
 	GPUProfilingData.Emplace();
+#endif
 }
 
 extern void DestroyShadersAndPrograms();
diff --git a/Engine/Source/Runtime/OpenGLDrv/Private/OpenGLDrv.h b/Engine/Source/Runtime/OpenGLDrv/Private/OpenGLDrv.h
--- a/Engine/Source/Runtime/OpenGLDrv/Private/OpenGLDrv.h
+++ b/Engine/Source/Runtime/OpenGLDrv/Private/OpenGLDrv.h
@@ -409,7 +409,7 @@ public:
 	virtual void RHIEndRenderQuery(FRHIRenderQuery* RenderQuery) final override;
 	virtual void RHIBeginDrawingViewport(FRHIViewport* Viewport, FRHITexture* RenderTargetRHI) final override;
 	virtual void RHIEndDrawingViewport(FRHIViewport* Viewport, bool bPresent, bool bLockToVsync) final override;
-	virtual void RHIEndFrame() final override;
+	virtual void RHIEndFrame(const FRHIEndFrameArgs& Args) final override;
 	virtual void RHISetStreamSource(uint32 StreamIndex, FRHIBuffer* VertexBuffer, uint32 Offset) final override;
 	virtual void RHISetRasterizerState(FRHIRasterizerState* NewState) final override;
 	virtual void RHISetViewport(float MinX, float MinY, float MinZ, float MaxX, float MaxY, float MaxZ) final override;
diff --git a/Engine/Source/Runtime/RHI/Private/GPUProfiler.cpp b/Engine/Source/Runtime/RHI/Private/GPUProfiler.cpp
--- a/Engine/Source/Runtime/RHI/Private/GPUProfiler.cpp
+++ b/Engine/Source/Runtime/RHI/Private/GPUProfiler.cpp
@@ -671,53 +671,507 @@ bool FGPUTiming::GAreGlobalsInitialized = false;
 
 namespace UE::RHI::GPUProfiler
 {
-	static TArray<IEventSink*>& GetSinks()
+	static TArray<FEventSink*>& GetSinks()
 	{
-		static TArray<IEventSink*> Sinks;
+		static TArray<FEventSink*> Sinks;
 		return Sinks;
 	}
 
-	void RegisterEventSink(IEventSink* Sink)
+	FEventSink::FEventSink()
 	{
-		GetSinks().Add(Sink);
+		GetSinks().Add(this);
 	}
 
-	void PushEvents(
-		FQueue Queue,
-		TConstArrayView<TUniquePtr<FBreadcrumbEvent>> BreadcrumbEvents,
-		TConstArrayView<TUniquePtr<FWorkEvent      >> WorkEvents,
-		TConstArrayView<TUniquePtr<FMarkerEvent    >> MarkerEvents)
+	FEventSink::~FEventSink()
 	{
-		if (BreadcrumbEvents.IsEmpty()
-		   || WorkEvents.IsEmpty()
-		   || MarkerEvents.IsEmpty())
+		GetSinks().RemoveSingle(this);
+	}
+
+	void ProcessEvents(FQueue Queue, TConstArrayView<TUniquePtr<UE::RHI::GPUProfiler::FEvent>> Events)
+	{
+		if (!Events.IsEmpty())
 		{
-			return;
+			for (FEventSink* Sink : GetSinks())
+			{
+				Sink->ProcessEvents(Queue, Events);
+			}
 		}
+	}
 
-		for (IEventSink* Sink : GetSinks())
+	void InitializeQueues(TConstArrayView<FQueue> Queues)
+	{
+		for (FEventSink* Sink : GetSinks())
 		{
-			Sink->ProcessEvents(Queue, BreadcrumbEvents, WorkEvents, MarkerEvents);
+			Sink->InitializeQueues(Queues);
 		}
 	}
 
-	struct FGPUProfilerSink_ProfileGPU final : public IEventSink
+	struct FGPUProfilerSink_StatUnit final : public FEventSink
 	{
-		FGPUProfilerSink_ProfileGPU()
+		struct FQueueTimestamps : public TArray<uint64>
 		{
-			RegisterEventSink(this);
+			int32 TimestampIndex = 0;
+			uint64 BusyCycles = 0;
+
+			uint64 GetCurrentTimestamp (uint64 Anchor) const { return (*this)[TimestampIndex] - Anchor; }
+			uint64 GetPreviousTimestamp(uint64 Anchor) const { return (*this)[TimestampIndex - 1] - Anchor; }
+
+			bool HasMoreTimestamps() const { return TimestampIndex < Num(); }
+			bool IsStartingWork()    const { return (TimestampIndex & 0x01) == 0x00; }
+
+			void AdvanceTimestamp() { TimestampIndex++; }
+		};
+
+		struct FQueueState
+		{
+			bool bBusy = false;
+			FQueueTimestamps Timestamps;
+		};
+
+		using FFrameState = TMap<FQueue, FQueueTimestamps>;
+
+		TMap<FQueue, FQueueState> QueueStates;
+		TMap<uint32, FFrameState> Frames;
+
+		void InitializeQueues(TConstArrayView<FQueue> Queues) override
+		{
+			for (FQueue const& Queue : Queues)
+			{
+				check(QueueStates.Find(Queue) == nullptr);
+				QueueStates.Add(Queue);
+			}
 		}
 
-		void ProcessEvents(
-			FQueue Queue,
-			TConstArrayView<TUniquePtr<FBreadcrumbEvent>> const& BreadcrumbEvents,
-			TConstArrayView<TUniquePtr<FWorkEvent      >> const& WorkEvents,
-			TConstArrayView<TUniquePtr<FMarkerEvent    >> const& MarkerEvents) override
+		void ProcessEvents(FQueue Queue, TConstArrayView<TUniquePtr<UE::RHI::GPUProfiler::FEvent>> Events) override
 		{
-			// @todo - new gpu profiler
+			FQueueState& QueueState = QueueStates.FindChecked(Queue);
+
+			for (auto const& Event : Events)
+			{
+				switch (Event->GetType())
+				{
+				case FEvent::EType::BeginWork:
+					{
+						check(!QueueState.bBusy);
+						QueueState.bBusy = true;
+						uint64 Value = Event->Value.Get<FEvent::FBeginWork>().GPUTimestampTOP;
+							
+						if (!QueueState.Timestamps.IsEmpty() && Value <= QueueState.Timestamps.Last())
+						{
+							//
+							// The Begin TOP event is sooner than the last End BOP event.
+							// The markers overlap, and the GPU was not idle.
+							// 
+							// Remove the previous End event, and discard this Begin event.
+							//
+							QueueState.Timestamps.RemoveAt(QueueState.Timestamps.Num() - 1, EAllowShrinking::No);
+						}
+						else
+						{
+							// GPU was idle. Keep this timestamp.
+							QueueState.Timestamps.Add(Value);
+						}
+					}
+					break;
+
+				case FEvent::EType::EndWork:
+					{
+						check(QueueState.bBusy);
+						QueueState.bBusy = false;
+
+						uint64 Value = Event->Value.Get<FEvent::FEndWork>().GPUTimestampBOP;
+						QueueState.Timestamps.Add(Value);
+					}
+					break;
+
+				case FEvent::EType::FrameBoundary:
+					{
+						check(!QueueState.bBusy);
+						auto const& FrameBoundary = Event->Value.Get<FEvent::FFrameBoundary>();
+
+						FFrameState& FrameState = Frames.FindOrAdd(FrameBoundary.FrameNumber);
+						FrameState.Emplace(Queue, MoveTemp(QueueState.Timestamps));
+
+						if (FrameState.Num() == QueueStates.Num())
+						{
+							// All registered queues have reported their frame boundary event.
+							// We have a full set of data to compute the total frame GPU stats.
+							ProcessFrame(FrameState);
+
+							Frames.Remove(FrameBoundary.FrameNumber);
+						}
+					}
+					break;
+				}
+			}
+		}
+
+		void ProcessFrame(FFrameState& FrameState)
+		{
+			// The total number of cycles where at least one GPU pipe was busy during the frame.
+			uint64 UnionBusyCycles = 0;
+
+			uint64 LastMinCycles = 0;
+			int32 BusyPipes = 0;
+			bool bFirst = true;
+
+			uint64 Anchor = 0; // @todo - handle possible timestamp wraparound
+
+			// Process the time ranges from each pipe.
+			while (true)
+			{
+				// Find the next minimum timestamp
+				FQueueTimestamps* NextMin = nullptr;
+				for (auto& [Queue, Current] : FrameState)
+				{
+					if (Current.HasMoreTimestamps() && (!NextMin || Current.GetCurrentTimestamp(Anchor) < NextMin->GetCurrentTimestamp(Anchor)))
+					{
+						NextMin = &Current;
+					}
+				}
+
+				if (!NextMin)
+					break; // No more timestamps to process
+
+				if (!bFirst)
+				{
+					if (BusyPipes > 0 && NextMin->GetCurrentTimestamp(Anchor) > LastMinCycles)
+					{
+						// Accumulate the union busy time across all pipes
+						UnionBusyCycles += NextMin->GetCurrentTimestamp(Anchor) - LastMinCycles;
+					}
+
+					if (!NextMin->IsStartingWork())
+					{
+						// Accumulate the busy time for this pipe specifically.
+						NextMin->BusyCycles += NextMin->GetCurrentTimestamp(Anchor) - NextMin->GetPreviousTimestamp(Anchor);
+					}
+				}
+
+				LastMinCycles = NextMin->GetCurrentTimestamp(Anchor);
+
+				BusyPipes += NextMin->IsStartingWork() ? 1 : -1;
+				check(BusyPipes >= 0);
+
+				NextMin->AdvanceTimestamp();
+				bFirst = false;
+			}
+
+			check(BusyPipes == 0);
+
+			// Update the global GPU frame time stats - need to convert to Cycles32 rather than Cycles64.
+			GGPUFrameTime = FPlatformMath::TruncToInt(FPlatformTime::ToSeconds64(UnionBusyCycles) / FPlatformTime::GetSecondsPerCycle());
+		}
+
+	} GGPUProfilerSink_StatUnit;
+
+	struct FGPUProfilerSink_ProfileGPU final : public FEventSink
+	{
+		struct FNode
+		{
+			FString Name;
+			FNode* Parent = nullptr;
+			FNode* Next = nullptr;
+			uint32 Level = 0;
+			TArray<uint64> Timestamps;
+
+			uint64 BusyCycles = 0;
+
+			FNode(FString&& Name)
+				: Name(MoveTemp(Name))
+			{}
+		};
+
+		struct FQueueState
+		{
+			FQueue const Queue;
+			TArray<TUniquePtr<FNode>> Nodes;
+			FNode* Current = nullptr;
+			FNode* Prev = nullptr;
+			FNode* First = nullptr;
+
+			enum class EState
+			{
+				Idle,
+				WaitingFrame,
+				Active
+			} State = EState::Idle;
+
+			std::atomic<bool> bTriggerProfile { false };
+
+			FQueueState(FQueue Queue)
+				: Queue(Queue)
+			{}
+
+			void PushNode(FString&& Name)
+			{
+				FNode* Parent = Current;
+				Current = Nodes.Emplace_GetRef(MakeUnique<FNode>(MoveTemp(Name))).Get();
+				Current->Parent = Parent;
+
+				if (!First)
+				{
+					First = Current;
+				}
+
+				if (Parent)
+				{
+					Current->Level = Parent->Level + 1;
+				}
+
+				if (Prev)
+				{
+					Prev->Next = Current;
+				}
+				Prev = Current;
+			}
+
+			void PopNode()
+			{
+				check(Current && Current->Parent);
+				Current = Current->Parent;
+			}
+
+			void ProcessEvents(TConstArrayView<TUniquePtr<UE::RHI::GPUProfiler::FEvent>> Events);
+
+			void LogTree(uint32 FrameNumber);
+		};
+
+		TMap<FQueue, TUniquePtr<FQueueState>> QueueStates;
+
+		void InitializeQueues(TConstArrayView<FQueue> Queues) override
+		{
+			for (FQueue const& Queue : Queues)
+			{
+				check(QueueStates.Find(Queue) == nullptr);
+				QueueStates.Add(Queue, MakeUnique<FQueueState>(Queue));
+			}
+		}
+
+		void ProcessEvents(FQueue Queue, TConstArrayView<TUniquePtr<UE::RHI::GPUProfiler::FEvent>> Events) override
+		{
+			QueueStates.FindChecked(Queue)->ProcessEvents(Events);
+		}
+
+		void ProfileNextFrame()
+		{
+			for (auto& [Queue, State] : QueueStates)
+			{
+				State->bTriggerProfile = true;
+			}
 		}
 
 	} GGPUProfilerSink_ProfileGPU;
+
+	void FGPUProfilerSink_ProfileGPU::FQueueState::ProcessEvents(TConstArrayView<TUniquePtr<UE::RHI::GPUProfiler::FEvent>> Events)
+	{
+		int32 Index = 0;
+	Restart:
+		if (State == EState::Idle && bTriggerProfile.exchange(false))
+		{
+			State = EState::WaitingFrame;
+		}
+
+		if (State == EState::WaitingFrame)
+		{
+			// Discard all received events until we reach a FrameBoundary event
+			for (; Index < Events.Num(); ++Index)
+			{
+				if (Events[Index]->GetType() == FEvent::EType::FrameBoundary)
+				{
+					// Start profiling until we receive another FrameBoundary event
+					State = EState::Active;
+
+					FEvent::FFrameBoundary& FrameBoundary = Events[Index]->Value.Get<FEvent::FFrameBoundary>();
+
+					// Build the node tree 
+					PushNode(TEXT("<root>"));
+
+				#if WITH_RHI_BREADCRUMBS
+					auto Recurse = [&](auto& Recurse, FRHIBreadcrumbNode* Breadcrumb) -> void
+					{
+						if (!Breadcrumb)
+						{
+							return;
+						}
+
+						Recurse(Recurse, Breadcrumb->GetParent());
+
+						FRHIBreadcrumb::FBuffer Buffer;
+						PushNode(Breadcrumb->Name.GetTCHAR(Buffer));
+					};
+					Recurse(Recurse, FrameBoundary.Breadcrumb);
+				#endif // WITH_RHI_BREADCRUMBS
+
+					++Index;
+					break;
+				}
+			}
+		}
+
+		if (State == EState::Active)
+		{
+			for (; Index < Events.Num(); ++Index)
+			{
+				auto const& Event = Events[Index];
+
+				switch (Event->GetType())
+				{
+				case FEvent::EType::BeginWork:
+					{
+						uint64 Timestamp = Event->Value.Get<FEvent::FBeginWork>().GPUTimestampTOP;
+						for (FNode* Node = Current; Node; Node = Node->Parent)
+						{
+							Node->Timestamps.Add(Timestamp);
+						}
+					}
+					break;
+
+				case FEvent::EType::EndWork:
+					{
+						uint64 Timestamp = Event->Value.Get<FEvent::FEndWork>().GPUTimestampBOP;
+						for (FNode* Node = Current; Node; Node = Node->Parent)
+						{
+							Node->Timestamps.Add(Timestamp);
+						}
+					}
+					break;
+
+			#if WITH_RHI_BREADCRUMBS
+				case FEvent::EType::BeginBreadcrumb:
+					{
+						auto& BeginBreadcrumb = Event->Value.Get<FEvent::FBeginBreadcrumb>();
+
+						// Push a new node
+						FRHIBreadcrumb::FBuffer Buffer;
+						PushNode(BeginBreadcrumb.Breadcrumb->Name.GetTCHAR(Buffer));
+
+						Current->Timestamps.Add(BeginBreadcrumb.GPUTimestampTOP);
+					}
+					break;
+
+				case FEvent::EType::EndBreadcrumb:
+					{
+						auto& EndBreadcrumb = Event->Value.Get<FEvent::FEndBreadcrumb>();
+						Current->Timestamps.Add(EndBreadcrumb.GPUTimestampBOP);
+
+						PopNode();
+					}
+					break;
+			#endif // WITH_RHI_BREADCRUMBS
+
+				case FEvent::EType::FrameBoundary:
+					{
+						FEvent::FFrameBoundary& FrameBoundary = Event->Value.Get<FEvent::FFrameBoundary>();
+						LogTree(FrameBoundary.FrameNumber);
+
+						// Reset tracking
+						Nodes.Reset();
+						Current = nullptr;
+						Prev = nullptr;
+						First = nullptr;
+
+						State = EState::Idle;
+						goto Restart;
+					}
+					break;
+				}
+			}
+		}
+	}
+
+	template <uint32 Width>
+	struct TUnicodeHorizontalBar
+	{
+		TCHAR Text[Width + 1];
+
+		// 0 <= Value <= 1
+		TUnicodeHorizontalBar(double Value)
+		{
+			TCHAR* Output = Text;
+			int32 Solid, Partial, Blank;
+			{
+				double Integer;
+				double Remainder = FMath::Modf(FMath::Clamp(Value, 0.0, 1.0) * Width, &Integer);
+
+				Solid = (int32)Integer;
+				Partial = (int32)FMath::Floor(Remainder * 8);
+				Blank = (Width - Solid - (Partial > 0 ? 1 : 0));
+			}
+
+			// Solid characters
+			for (int32 Index = 0; Index < Solid; ++Index)
+			{
+				*Output++ = TEXT('█');
+			}
+
+			// Partially filled character
+			if (Partial > 0)
+			{
+				static constexpr TCHAR const Data[] = TEXT("▏▎▍▌▋▊▉");
+				*Output++ = Data[Partial - 1];
+			}
+
+			// Blank Characters to pad out the width
+			for (int32 Index = 0; Index < Blank; ++Index)
+			{
+				*Output++ = TEXT(' ');
+			}
+
+			*Output++ = 0;
+			check(uintptr_t(Output) == (uintptr_t(Text) + sizeof(Text)));
+		}
+	};
+
+	void FGPUProfilerSink_ProfileGPU::FQueueState::LogTree(uint32 FrameNumber)
+	{
+		for (FNode* Node = First; Node; Node = Node->Next)
+		{
+			check(Node->Timestamps.Num() % 2 == 0);
+			Node->BusyCycles = 0;
+
+			uint64 LastBeginCycles = 0;
+			uint64 LastEndCycles = 0;
+
+			for (int32 Index = 0; Index < Node->Timestamps.Num(); ++Index)
+			{
+				if ((Index & 1) == 0)
+				{
+					// Begin
+					LastBeginCycles = FMath::Max(LastEndCycles, Node->Timestamps[Index]);
+				}
+				else
+				{
+					// End
+					uint64 End = Node->Timestamps[Index];
+					Node->BusyCycles += End - LastBeginCycles;
+
+					LastEndCycles = End;
+				}
+			}
+		}
+
+		double RootMilliseconds = FPlatformTime::ToMilliseconds64(First->BusyCycles);
+
+		FString LogMessage;
+		for (FNode* Node = First; Node; Node = Node->Next)
+		{
+			double Milliseconds = FPlatformTime::ToMilliseconds64(Node->BusyCycles);
+
+			TUnicodeHorizontalBar<8> Bar = Milliseconds / RootMilliseconds;
+			LogMessage += FString::Printf(TEXT("%9.3f ms |%s| %*s\n"), Milliseconds, Bar.Text, Node->Name.Len() + (Node->Level * 4), *Node->Name);
+		}
+
+		UE_LOG(LogRHI, Display, TEXT("GPU Profile for Frame %d, Queue [%s, GPU: %d, Idx: %d]:\n%s\n\n"), FrameNumber, Queue.GetTypeString(), Queue.GPU, Queue.Index, *LogMessage);
+	}
+
+	static FAutoConsoleCommand ProfileGPUNew(
+		TEXT("ProfileGPUNew"),
+		TEXT(""),
+		FConsoleCommandWithArgsDelegate::CreateLambda([](const TArray<FString>& Args)
+		{
+			GGPUProfilerSink_ProfileGPU.ProfileNextFrame();
+		}));
 }
 
 #endif // RHI_NEW_GPU_PROFILER
diff --git a/Engine/Source/Runtime/RHI/Private/RHICommandList.cpp b/Engine/Source/Runtime/RHI/Private/RHICommandList.cpp
--- a/Engine/Source/Runtime/RHI/Private/RHICommandList.cpp
+++ b/Engine/Source/Runtime/RHI/Private/RHICommandList.cpp
@@ -974,7 +974,7 @@ FGraphEventRef FRHICommandListExecutor::FTranslateState::Finalize()
 	}
 }
 
-void FRHICommandListExecutor::FSubmitState::Submit()
+void FRHICommandListExecutor::FSubmitState::Submit(const FSubmitArgs& Args)
 {
 	// Coalesce finalized platform command lists into a single array
 	TArray<IRHIPlatformCommandList*> FinalizedCmdLists;
@@ -1060,7 +1060,14 @@ void FRHICommandListExecutor::FSubmitState::Submit()
 
 	if (EnumHasAnyFlags(SubmitFlags, ERHISubmitFlags::EndFrame))
 	{
-		GDynamicRHI->RHIEndFrame();
+		FDynamicRHI::FRHIEndFrameArgs EndFrameArgs
+		{
+#if WITH_RHI_BREADCRUMBS
+			.GPUBreadcrumbs = Args.GPUBreadcrumbs
+#endif
+		};
+
+		GDynamicRHI->RHIEndFrame(EndFrameArgs);
 		GRHICommandList.FrameDrawStats.ProcessAsFrameStats();
 	}
 
@@ -1210,12 +1217,20 @@ RHI_API void FRHICommandListExecutor::Submit(TConstArrayView<FRHICommandListBase
 				// Finalize the last translate job
 				State->FinalizeCurrent();
 
+				FSubmitState::FSubmitArgs Args;
+#if WITH_RHI_BREADCRUMBS
+				for (ERHIPipeline Pipeline : MakeFlagsRange(ERHIPipeline::All))
+				{
+					Args.GPUBreadcrumbs[Pipeline] = GRHICommandList.Breadcrumbs.GPU[Pipeline].Current;
+				}
+#endif
+
 				// Submission thread
 				EnqueueSubmitTask(MoveTemp(State->TranslateEvents),
-					[State]() mutable
+					[State, Args]() mutable
 					{
 						SCOPED_NAMED_EVENT(RHI_SubmitToGPU, FColor::White);
-						State->Submit();
+						State->Submit(Args);
 					}
 				);
 				LastSubmit = RHIThreadPipe.Close();
@@ -1341,6 +1356,8 @@ FGraphEventRef FRHICommandListBase::RHIThreadFence(bool bSetLockFence)
 {
 	checkf(IsTopOfPipe() || Bypass(), TEXT("RHI thread fences only work when recording RHI commands (or in bypass mode)."));
 
+	bUsesLockFence |= bSetLockFence;
+
 	if (IsRunningRHIInSeparateThread())
 	{
 		FGraphEventRef Fence = nullptr;
diff --git a/Engine/Source/Runtime/RHI/Private/RHIGlobals.cpp b/Engine/Source/Runtime/RHI/Private/RHIGlobals.cpp
--- a/Engine/Source/Runtime/RHI/Private/RHIGlobals.cpp
+++ b/Engine/Source/Runtime/RHI/Private/RHIGlobals.cpp
@@ -5,6 +5,8 @@
 
 FRHIGlobals GRHIGlobals;
 
+uint32 GGPUFrameTime = 0;
+
 //
 //	MSAA sample offsets.
 //
diff --git a/Engine/Source/Runtime/RHI/Private/RHIValidation.cpp b/Engine/Source/Runtime/RHI/Private/RHIValidation.cpp
--- a/Engine/Source/Runtime/RHI/Private/RHIValidation.cpp
+++ b/Engine/Source/Runtime/RHI/Private/RHIValidation.cpp
@@ -966,10 +966,10 @@ void FValidationRHI::RHIEndFrame_RenderThread(FRHICommandListImmediate& RHICmdLi
 	RHI->RHIEndFrame_RenderThread(RHICmdList);
 }
 
-void FValidationRHI::RHIEndFrame()
+void FValidationRHI::RHIEndFrame(const FRHIEndFrameArgs& Args)
 {
 	RHIThreadFrameID++;
-	RHI->RHIEndFrame();
+	RHI->RHIEndFrame(Args);
 }
 
 namespace RHIValidation
diff --git a/Engine/Source/Runtime/RHI/Public/DynamicRHI.h b/Engine/Source/Runtime/RHI/Public/DynamicRHI.h
--- a/Engine/Source/Runtime/RHI/Public/DynamicRHI.h
+++ b/Engine/Source/Runtime/RHI/Public/DynamicRHI.h
@@ -152,7 +152,14 @@ public:
 	/////// RHI Methods
 
 	RHI_API virtual void RHIEndFrame_RenderThread(FRHICommandListImmediate& RHICmdList);
-	virtual void RHIEndFrame() = 0;
+
+	struct FRHIEndFrameArgs
+	{
+#if WITH_RHI_BREADCRUMBS
+		const TRHIPipelineArray<FRHIBreadcrumbNode*>& GPUBreadcrumbs;
+#endif
+	};
+	virtual void RHIEndFrame(const FRHIEndFrameArgs& Args) = 0;
 
 	// FlushType: Thread safe
 	virtual FSamplerStateRHIRef RHICreateSamplerState(const FSamplerStateInitializerRHI& Initializer) = 0;
diff --git a/Engine/Source/Runtime/RHI/Public/GPUProfiler.h b/Engine/Source/Runtime/RHI/Public/GPUProfiler.h
--- a/Engine/Source/Runtime/RHI/Public/GPUProfiler.h
+++ b/Engine/Source/Runtime/RHI/Public/GPUProfiler.h
@@ -24,18 +24,77 @@ namespace UE::RHI::GPUProfiler
 			SwapChain
 		};
 
-		EType Type;
-		uint8 GPU;
-		uint8 Index;
+		union
+		{
+			struct
+			{
+				EType Type;
+				uint8 GPU;
+				uint8 Index;
+				uint8 Padding;
+			};
+			uint32 Value = 0;
+		};
+
+		FQueue() = default;
+
+		FQueue(EType Type, uint8 GPU, uint8 Index)
+			: Type   (Type)
+			, GPU    (GPU)
+			, Index  (Index)
+			, Padding(0)
+		{}
+
+		bool operator == (FQueue const& RHS) const
+		{
+			return Value == RHS.Value;
+		}
+
+		bool operator != (FQueue const& RHS) const
+		{
+			return !(*this == RHS);
+		}
+
+		friend uint32 GetTypeHash(FQueue const& Queue)
+		{
+			return GetTypeHash(Queue.Value);
+		}
+
+		TCHAR const* GetTypeString() const
+		{
+			switch (Type)
+			{
+			case EType::Graphics:  return TEXT("Graphics");
+			case EType::Compute:   return TEXT("Compute");
+			case EType::Copy:      return TEXT("Copy");
+			case EType::SwapChain: return TEXT("Swapchain");
+			default:               return TEXT("<unknown>");
+			}
+		}
 	};
 
-	struct FBreadcrumbEvent
+	struct FEvent
 	{
+		//
+		// All timestamps are relative to FPlatformTime::Cycles64().
+		// TOP = Top of Pipe. Timestamps written by the GPU's command processor before work begins.
+		// BOP = Bottom of Pipe. Timestamps written after the GPU completes work.
+		//
+		
+		// Inserted on each call to RHIEndFrame. Marks the end of a profiler frame.
 		struct FFrameBoundary
 		{
+			// The index of the frame that just ended.
+			// Very first frame of the engine is frame 0 (from boot to first call to RHIEndFrame).
 			uint32 FrameNumber;
+
+#if WITH_RHI_BREADCRUMBS
+			// The RHI breadcrumb currently at the top of the stack at the frame boundary.
+			FRHIBreadcrumbNode* Breadcrumb;
+#endif
 		};
 
+#if WITH_RHI_BREADCRUMBS
 		struct FBeginBreadcrumb
 		{
 			FRHIBreadcrumbNode* Breadcrumb;
@@ -47,27 +106,19 @@ namespace UE::RHI::GPUProfiler
 			FRHIBreadcrumbNode* Breadcrumb;
 			uint64 GPUTimestampBOP;
 		};
-		
-		TVariant<
-			  FFrameBoundary
-			, FBeginBreadcrumb
-			, FEndBreadcrumb
-		> Value;
-
-		template <typename T>
-		FBreadcrumbEvent(T const& Value)
-			: Value(TInPlaceType<T>(), Value)
-		{}
-	};
+#endif // WITH_RHI_BREADCRUMBS
 
-	struct FWorkEvent
-	{
+		// Inserted when the GPU starts work on a queue.
 		struct FBeginWork
 		{
+			// CPU timestamp of when the work was submitted to the driver for execution on the GPU.
 			uint64 CPUTimestamp;
+
+			// TOP timestamp of when the work actually started on the GPU.
 			uint64 GPUTimestampTOP;
 		};
 
+		// Inserted when the GPU completes work on a queue and goes idle.
 		struct FEndWork
 		{
 			uint64 GPUTimestampBOP;
@@ -85,21 +136,6 @@ namespace UE::RHI::GPUProfiler
 			uint64 ID;
 		};
 
-		TVariant<
-			  FBeginWork
-			, FEndWork
-			, FSignalFence
-			, FWaitFence
-		> Value;
-
-		template <typename T>
-		FWorkEvent(T const& Value)
-			: Value(TInPlaceType<T>(), Value)
-		{}
-	};
-
-	struct FMarkerEvent
-	{
 		struct FFlip
 		{
 			uint64 GPUTimestamp;
@@ -109,35 +145,65 @@ namespace UE::RHI::GPUProfiler
 		{
 			uint64 GPUTimestamp;
 		};
-
-		TVariant<
-			  FFlip
+		
+		using FStorage = TVariant<
+			  FFrameBoundary
+#if WITH_RHI_BREADCRUMBS
+			, FBeginBreadcrumb
+			, FEndBreadcrumb
+#endif
+			, FBeginWork
+			, FEndWork
+			, FSignalFence
+			, FWaitFence
+			, FFlip
 			, FVsync
-		> Value;
+		>;
+
+		enum class EType
+		{
+			FrameBoundary   = FStorage::IndexOfType<FFrameBoundary  >(),
+#if WITH_RHI_BREADCRUMBS
+			BeginBreadcrumb = FStorage::IndexOfType<FBeginBreadcrumb>(),
+			EndBreadcrumb   = FStorage::IndexOfType<FEndBreadcrumb  >(),
+#endif
+			BeginWork       = FStorage::IndexOfType<FBeginWork      >(),
+			EndWork         = FStorage::IndexOfType<FEndWork        >(),
+			SignalFence     = FStorage::IndexOfType<FSignalFence    >(),
+			WaitFence       = FStorage::IndexOfType<FWaitFence      >(),
+			Flip            = FStorage::IndexOfType<FFlip           >(),
+			VSync		    = FStorage::IndexOfType<FVsync          >()
+		};
+
+		FStorage Value;
+
+		EType GetType() const
+		{
+			return static_cast<EType>(Value.GetIndex());
+		}
 
 		template <typename T>
-		FMarkerEvent(T const& Value)
+		FEvent(T const& Value)
 			: Value(TInPlaceType<T>(), Value)
 		{}
 	};
 
-	struct IEventSink
+	struct FEventSink
 	{
-		virtual void ProcessEvents(
-			FQueue Queue,
-			TConstArrayView<TUniquePtr<FBreadcrumbEvent>> const& BreadcrumbEvents,
-			TConstArrayView<TUniquePtr<FWorkEvent      >> const& WorkEvents,
-			TConstArrayView<TUniquePtr<FMarkerEvent    >> const& MarkerEvents) = 0;
-	};
+	protected:
+		RHI_API FEventSink();
+		RHI_API ~FEventSink();
 
-	RHI_API void RegisterEventSink(IEventSink* Sink);
+		FEventSink(FEventSink const&) = delete;
+		FEventSink(FEventSink&&) = delete;
+
+	public:
+		virtual void ProcessEvents(FQueue Queue, TConstArrayView<TUniquePtr<UE::RHI::GPUProfiler::FEvent>> Events) = 0;
+		virtual void InitializeQueues(TConstArrayView<FQueue> Queues) = 0;
+	};
 
-	RHI_API void PushEvents(
-		FQueue Queue,
-		TConstArrayView<TUniquePtr<FBreadcrumbEvent>> BreadcrumbEvents,
-		TConstArrayView<TUniquePtr<FWorkEvent      >> WorkEvents,
-		TConstArrayView<TUniquePtr<FMarkerEvent    >> MarkerEvents
-	);
+	RHI_API void ProcessEvents(FQueue Queue, TConstArrayView<TUniquePtr<UE::RHI::GPUProfiler::FEvent>> Events);
+	RHI_API void InitializeQueues(TConstArrayView<FQueue> Queues);
 }
 
 #else
diff --git a/Engine/Source/Runtime/RHI/Public/RHICommandList.h b/Engine/Source/Runtime/RHI/Public/RHICommandList.h
--- a/Engine/Source/Runtime/RHI/Public/RHICommandList.h
+++ b/Engine/Source/Runtime/RHI/Public/RHICommandList.h
@@ -472,7 +472,7 @@ public:
 		}
 
 		// Command lists that use RHIThreadFence(true) are going to mutate resource state, so must be single-threaded.
-		if (LastLockFenceCommand)
+		if (bUsesLockFence)
 		{
 			return false;
 		}
@@ -1219,6 +1219,7 @@ protected:
 	bool bAllowParallelTranslate = true;
 	bool bUsesSetTrackedAccess   = false;
 	bool bUsesShaderBundles      = false;
+	bool bUsesLockFence          = false;
 
 	// The currently selected pipelines that RHI commands are directed to, during command list recording.
 	// This is also adjusted during command list execution based on recorded use of ActivatePipeline().
@@ -5336,7 +5337,13 @@ private:
 
 		void Dispatch(FRHICommandListBase* CmdList);
 
-		void Submit();
+		struct FSubmitArgs
+		{
+#if WITH_RHI_BREADCRUMBS
+			TRHIPipelineArray<FRHIBreadcrumbNode*> GPUBreadcrumbs;
+#endif
+		};
+		void Submit(const FSubmitArgs& Args);
 		void FinalizeCurrent();
 	} *SubmitState = nullptr;
 
diff --git a/Engine/Source/Runtime/RHI/Public/RHIGlobals.h b/Engine/Source/Runtime/RHI/Public/RHIGlobals.h
--- a/Engine/Source/Runtime/RHI/Public/RHIGlobals.h
+++ b/Engine/Source/Runtime/RHI/Public/RHIGlobals.h
@@ -680,6 +680,9 @@ struct FRHIGlobals
 
 extern RHI_API FRHIGlobals GRHIGlobals;
 
+/** The GPU time taken to render the last frame. Same metric as FPlatformTime::Cycles(). */
+extern RHI_API uint32 GGPUFrameTime;
+
 //
 // Deprecated old-style names
 
diff --git a/Engine/Source/Runtime/RHI/Public/RHIValidation.h b/Engine/Source/Runtime/RHI/Public/RHIValidation.h
--- a/Engine/Source/Runtime/RHI/Public/RHIValidation.h
+++ b/Engine/Source/Runtime/RHI/Public/RHIValidation.h
@@ -101,7 +101,7 @@ public:
 	/////// RHI Methods
 
 	virtual void RHIEndFrame_RenderThread(FRHICommandListImmediate& RHICmdList) override final;
-	virtual void RHIEndFrame() override final;
+	virtual void RHIEndFrame(const FRHIEndFrameArgs& Args) override final;
 
 	// FlushType: Thread safe
 	virtual FSamplerStateRHIRef RHICreateSamplerState(const FSamplerStateInitializerRHI& Initializer) override final
diff --git a/Engine/Source/Runtime/VulkanRHI/Private/VulkanRHI.cpp b/Engine/Source/Runtime/VulkanRHI/Private/VulkanRHI.cpp
--- a/Engine/Source/Runtime/VulkanRHI/Private/VulkanRHI.cpp
+++ b/Engine/Source/Runtime/VulkanRHI/Private/VulkanRHI.cpp
@@ -1151,7 +1151,7 @@ void FVulkanDynamicRHI::RHIEndFrame_RenderThread(FRHICommandListImmediate& RHICm
 	});
 }
 
-void FVulkanDynamicRHI::RHIEndFrame()
+void FVulkanDynamicRHI::RHIEndFrame(const FRHIEndFrameArgs& Args)
 {
 	// @todo dev-pr - refactor RHIEndFrame_RenderThread to reduce use of the immediate command list, and move cleanup work to here.
 }
diff --git a/Engine/Source/Runtime/VulkanRHI/Public/VulkanDynamicRHI.h b/Engine/Source/Runtime/VulkanRHI/Public/VulkanDynamicRHI.h
--- a/Engine/Source/Runtime/VulkanRHI/Public/VulkanDynamicRHI.h
+++ b/Engine/Source/Runtime/VulkanRHI/Public/VulkanDynamicRHI.h
@@ -98,7 +98,7 @@ public:
 	void InitInstance();
 
 	virtual void RHIEndFrame_RenderThread(FRHICommandListImmediate& RHICmdList) final override;
-	virtual void RHIEndFrame() final override;
+	virtual void RHIEndFrame(const FRHIEndFrameArgs& Args) final override;
 
 	virtual FSamplerStateRHIRef RHICreateSamplerState(const FSamplerStateInitializerRHI& Initializer) final override;
 	virtual FRasterizerStateRHIRef RHICreateRasterizerState(const FRasterizerStateInitializerRHI& Initializer) final override;
diff --git a/Engine/Source/Runtime/Windows/D3D11RHI/Private/D3D11RHI.cpp b/Engine/Source/Runtime/Windows/D3D11RHI/Private/D3D11RHI.cpp
--- a/Engine/Source/Runtime/Windows/D3D11RHI/Private/D3D11RHI.cpp
+++ b/Engine/Source/Runtime/Windows/D3D11RHI/Private/D3D11RHI.cpp
@@ -29,7 +29,7 @@ extern void UniformBufferBeginFrame();
 // Has to be .exe module to be correctly detected.
 // extern "C" { _declspec(dllexport) uint32 NvOptimusEnablement = 0x00000001; }
 
-void FD3D11DynamicRHI::RHIEndFrame()
+void FD3D11DynamicRHI::RHIEndFrame(const FRHIEndFrameArgs& Args)
 {
 	// End Frame
 #if (RHI_NEW_GPU_PROFILER == 0)
diff --git a/Engine/Source/Runtime/Windows/D3D11RHI/Private/D3D11RHIPrivate.h b/Engine/Source/Runtime/Windows/D3D11RHI/Private/D3D11RHIPrivate.h
--- a/Engine/Source/Runtime/Windows/D3D11RHI/Private/D3D11RHIPrivate.h
+++ b/Engine/Source/Runtime/Windows/D3D11RHI/Private/D3D11RHIPrivate.h
@@ -524,7 +524,7 @@ public:
 
 	virtual void RHIBeginDrawingViewport(FRHIViewport* Viewport, FRHITexture* RenderTargetRHI) final override;
 	virtual void RHIEndDrawingViewport(FRHIViewport* Viewport, bool bPresent, bool bLockToVsync) final override;
-	virtual void RHIEndFrame() final override;
+	virtual void RHIEndFrame(const FRHIEndFrameArgs& Args) final override;
 	virtual void RHISetStreamSource(uint32 StreamIndex, FRHIBuffer* VertexBuffer, uint32 Offset) final override;
 	virtual void RHISetRasterizerState(FRHIRasterizerState* NewState) final override;
 	virtual void RHISetViewport(float MinX, float MinY, float MinZ, float MaxX, float MaxY, float MaxZ) final override;
