Remove CommandListID field in D3D12 RHI
 - No longer used.

diff --git a/Engine/Source/Runtime/D3D12RHI/Private/D3D12CommandContext.h b/Engine/Source/Runtime/D3D12RHI/Private/D3D12CommandContext.h
--- a/Engine/Source/Runtime/D3D12RHI/Private/D3D12CommandContext.h
+++ b/Engine/Source/Runtime/D3D12RHI/Private/D3D12CommandContext.h
@@ -238,9 +238,6 @@ public:
 
 	bool IsOpen() const { return CommandList != nullptr; }
 
-	// Returns unique identity that can be used to distinguish between command lists even after they were recycled.
-	uint64 GetCommandListID() { return GetCommandList().State.CommandListID; }
-
 	FD3D12SyncPoint* GetContextSyncPoint()
 	{
 		if (!ContextSyncPoint)
diff --git a/Engine/Source/Runtime/D3D12RHI/Private/D3D12CommandList.cpp b/Engine/Source/Runtime/D3D12RHI/Private/D3D12CommandList.cpp
--- a/Engine/Source/Runtime/D3D12RHI/Private/D3D12CommandList.cpp
+++ b/Engine/Source/Runtime/D3D12RHI/Private/D3D12CommandList.cpp
@@ -16,9 +16,6 @@ static FAutoConsoleVariableRef CVarD3D12ExtraDepthTransitions(
 	GD3D12ExtraDepthTransitions,
 	TEXT("Adds extra transitions for the depth buffer to fix validation issues. However, this currently breaks async compute"));
 
-
-int64 FD3D12CommandList::FState::NextCommandListID = 0;
-
 void FD3D12CommandList::UpdateResidency(const FD3D12Resource* Resource)
 {
 #if ENABLE_RESIDENCY_MANAGEMENT
@@ -460,7 +457,6 @@ void FD3D12CommandList::WriteTimestamp(FD3D12QueryLocation const& Location, ED3D
 
 FD3D12CommandList::FState::FState(FD3D12CommandAllocator* CommandAllocator, FD3D12QueryAllocator* TimestampAllocator, FD3D12QueryAllocator* PipelineStatsAllocator)
 	: CommandAllocator(CommandAllocator)
-	, CommandListID   (FPlatformAtomics::InterlockedIncrement(&NextCommandListID))
 {
 	PendingResourceBarriers.Reserve(256);
 
diff --git a/Engine/Source/Runtime/D3D12RHI/Private/D3D12CommandList.h b/Engine/Source/Runtime/D3D12RHI/Private/D3D12CommandList.h
--- a/Engine/Source/Runtime/D3D12RHI/Private/D3D12CommandList.h
+++ b/Engine/Source/Runtime/D3D12RHI/Private/D3D12CommandList.h
@@ -270,8 +270,6 @@ private:
 	// Contents of the state struct are reset when the command list is recycled
 	struct FState
 	{
-		static int64 NextCommandListID;
-
 		FState(FD3D12CommandAllocator* CommandAllocator, FD3D12QueryAllocator* TimestampAllocator, FD3D12QueryAllocator* PipelineStatsAllocator);
 
 		// The allocator currently assigned to this command list.
@@ -283,11 +281,6 @@ private:
 		// A map of all D3D resources, and their states, that were state transitioned with tracking.
 		TMap<FD3D12Resource*, CResourceState> TrackedResourceState;
 
-		// Unique ID of this command list used to avoid costly redundant operations, such as resource residency updates.
-		// This value is updated every time the command list is reset, so it is safe to use even when command list object is recycled.
-		// Value should be only used for identity, not for synchronization. Valid values are guaranteed to be > 0.
-		uint64 CommandListID;
-
 #if DEBUG_RESOURCE_STATES
 		// Tracks all the resources barriers being issued on this command list in order
 		TArray<D3D12_RESOURCE_BARRIER> ResourceBarriers;
