[Insights] Fixed usage of TRACE_CPUPROFILER_EVENT_SCOPE / TRACE_CPUPROFILER_EVENT_SCOPE_STR.

#rnx

diff --git a/Engine/Source/Runtime/D3D12RHI/Private/D3D12BindlessDescriptors.cpp b/Engine/Source/Runtime/D3D12RHI/Private/D3D12BindlessDescriptors.cpp
--- a/Engine/Source/Runtime/D3D12RHI/Private/D3D12BindlessDescriptors.cpp
+++ b/Engine/Source/Runtime/D3D12RHI/Private/D3D12BindlessDescriptors.cpp
@@ -155,7 +155,7 @@ void FD3D12BindlessResourceManager::Free(FRHIDescriptorHandle InHandle)
 
 FRHIDescriptorHandle FD3D12BindlessResourceManager::ResizeGrowAndAllocate()
 {	
-	TRACE_CPUPROFILER_EVENT_SCOPE("FD3D12BindlessResourceManager::ResizeGrow");
+	TRACE_CPUPROFILER_EVENT_SCOPE(FD3D12BindlessResourceManager::ResizeGrow);
 
 	FScopeLock ScopeLock(&HeapsCS);
 
@@ -272,7 +272,7 @@ void FD3D12BindlessResourceManager::GarbageCollect()
 	}
 	else
 	{
-		TRACE_CPUPROFILER_EVENT_SCOPE("FD3D12BindlessResourceManager::GarbageCollect");
+		TRACE_CPUPROFILER_EVENT_SCOPE(FD3D12BindlessResourceManager::GarbageCollect);
 
 		// Update the moving window max gpu heaps and reset the working value
 		MovingWindowMaxInUseGPUHeaps.PushValue(MaxInUseGPUHeaps);
@@ -359,7 +359,7 @@ void FD3D12BindlessResourceManager::InitializeDescriptor(FRHIDescriptorHandle Ds
 {
 	if (DstHandle.IsValid())
 	{
-		TRACE_CPUPROFILER_EVENT_SCOPE("FD3D12BindlessResourceManager::InitializeDescriptor");
+		TRACE_CPUPROFILER_EVENT_SCOPE(FD3D12BindlessResourceManager::InitializeDescriptor);
 		
 		FScopeLock ScopeLock(&HeapsCS);
 
@@ -382,7 +382,7 @@ void FD3D12BindlessResourceManager::UpdateDescriptor(FD3D12ContextArray const& C
 {
 	if (DstHandle.IsValid())
 	{
-		TRACE_CPUPROFILER_EVENT_SCOPE("FD3D12BindlessResourceManager::UpdateDescriptor");
+		TRACE_CPUPROFILER_EVENT_SCOPE(FD3D12BindlessResourceManager::UpdateDescriptor);
 
 		check(IsInRHIThread() || GRHICommandList.Bypass());
 	
@@ -556,7 +556,7 @@ void FD3D12BindlessResourceManager::CheckRequestNewActiveGPUHeap()
 		return;
 	}
 
-	TRACE_CPUPROFILER_EVENT_SCOPE("FD3D12BindlessResourceManager::RequestNewActiveGPUHeap");
+	TRACE_CPUPROFILER_EVENT_SCOPE(FD3D12BindlessResourceManager::RequestNewActiveGPUHeap);
 
 	FScopeLock ScopeLock(&HeapsCS);
 
diff --git a/Engine/Source/Runtime/D3D12RHI/Private/D3D12Descriptors.cpp b/Engine/Source/Runtime/D3D12RHI/Private/D3D12Descriptors.cpp
--- a/Engine/Source/Runtime/D3D12RHI/Private/D3D12Descriptors.cpp
+++ b/Engine/Source/Runtime/D3D12RHI/Private/D3D12Descriptors.cpp
@@ -19,7 +19,7 @@ void UE::D3D12Descriptors::CopyDescriptor(FD3D12Device* Device, FD3D12Descriptor
 void UE::D3D12Descriptors::CopyDescriptors(FD3D12Device* Device, FD3D12DescriptorHeap* TargetHeap, FD3D12DescriptorHeap* SourceHeap, uint32 FirstDescriptor, uint32 NumDescriptors)
 {
 	// Use regular CPU Profile without dynamic value - almost same overhead for string formatting as actual Copy
-	TRACE_CPUPROFILER_EVENT_SCOPE("CopyDescriptors HeapToHeap");
+	TRACE_CPUPROFILER_EVENT_SCOPE_STR("CopyDescriptors HeapToHeap");
 	//SCOPED_NAMED_EVENT_F(TEXT("CopyDescriptors HeapToHeap (%d)"), FColor::Turquoise, NumDescriptors);
 
 	const D3D12_CPU_DESCRIPTOR_HANDLE TargetStart = TargetHeap->GetCPUSlotHandle(FirstDescriptor);
@@ -37,7 +37,7 @@ void UE::D3D12Descriptors::CopyDescriptors(FD3D12Device* Device, FD3D12Descripto
 void UE::D3D12Descriptors::CopyDescriptors(FD3D12Device* Device, FD3D12DescriptorHeap* TargetHeap, FD3D12DescriptorHeap* SourceHeap, TConstArrayView<FRHIDescriptorHandle> Handles)
 {
 	// Use regular CPU Profile without dynamic value - almost same overhead for string formatting as actual Copy
-	TRACE_CPUPROFILER_EVENT_SCOPE("CopyDescriptors HeapToHeap Scattered");
+	TRACE_CPUPROFILER_EVENT_SCOPE_STR("CopyDescriptors HeapToHeap Scattered");
 	//SCOPED_NAMED_EVENT_F(TEXT("CopyDescriptors HeapToHeap Scattered (%d)"), FColor::Turquoise, Handles.Num());
 
 	const int32 NumDescriptors = Handles.Num();
@@ -68,7 +68,7 @@ void UE::D3D12Descriptors::CopyDescriptors(FD3D12Device* Device, FD3D12Descripto
 void UE::D3D12Descriptors::CopyDescriptors(FD3D12Device* Device, FD3D12DescriptorHeap* TargetHeap, TConstArrayView<FRHIDescriptorHandle> DstHandles, TConstArrayView<FD3D12OfflineDescriptor> SrcOfflineDescriptors)
 {
 	// Use regular CPU Profile without dynamic value - almost same overhead for string formatting as actual Copy
-	TRACE_CPUPROFILER_EVENT_SCOPE("CopyDescriptors Scattered");
+	TRACE_CPUPROFILER_EVENT_SCOPE_STR("CopyDescriptors Scattered");
 	//SCOPED_NAMED_EVENT_F(TEXT("CopyDescriptors Scattered (%d)"), FColor::Turquoise, DstHandles.Num());
 
 	check(DstHandles.Num() == SrcOfflineDescriptors.Num());
diff --git a/Engine/Source/Runtime/MovieScene/Private/EntitySystem/MovieSceneTaskScheduler.cpp b/Engine/Source/Runtime/MovieScene/Private/EntitySystem/MovieSceneTaskScheduler.cpp
--- a/Engine/Source/Runtime/MovieScene/Private/EntitySystem/MovieSceneTaskScheduler.cpp
+++ b/Engine/Source/Runtime/MovieScene/Private/EntitySystem/MovieSceneTaskScheduler.cpp
@@ -537,7 +537,7 @@ void FEntitySystemScheduler::ExecuteTasks()
 		ensure(NumInitialTasks != 0);
 	}
 
-	TRACE_CPUPROFILER_EVENT_SCOPE("Wait For Scheduled Tasks");
+	TRACE_CPUPROFILER_EVENT_SCOPE_STR("Wait For Scheduled Tasks");
 	for(;;)
 	{
 		while (const FScheduledTask* Task = GameThreadTaskList.Pop())
