Shader bundles only need bindless parameters when full raytracing is enabled.

diff --git a/Engine/Source/Runtime/D3D12RHI/Private/D3D12Adapter.cpp b/Engine/Source/Runtime/D3D12RHI/Private/D3D12Adapter.cpp
--- a/Engine/Source/Runtime/D3D12RHI/Private/D3D12Adapter.cpp
+++ b/Engine/Source/Runtime/D3D12RHI/Private/D3D12Adapter.cpp
@@ -1234,7 +1234,7 @@ void FD3D12Adapter::InitializeDevices()
 			bBindlessResourcesAllowed = (BindlessResourcesConfig != ERHIBindlessConfiguration::Disabled);
 			bBindlessSamplersAllowed = (BindlessSamplersConfig != ERHIBindlessConfiguration::Disabled);
 
-			GRHIGlobals.ShaderBundles.RequiresSharedBindlessParameters = (bBindlessResourcesAllowed || bBindlessSamplersAllowed);
+			GRHIGlobals.ShaderBundles.RequiresSharedBindlessParameters = (BindlessResourcesConfig == ERHIBindlessConfiguration::AllShaders || BindlessSamplersConfig == ERHIBindlessConfiguration::AllShaders);
 		}
 #endif
 
