FD3D12ExplicitDescriptorHeapCache::AllocateHeap is now rounding up NumDescriptors to next power of 2 which means it will no longer match GD3D12ExplicitViewDescriptorHeapSize which will prevent GD3D12RayTracingViewDescriptorHeapOverflowReported from ever being set and as a result we won't reallocate descriptor cache for ray tracing. Change the condition to greater equal instead so GD3D12RayTracingViewDescriptorHeapOverflowReported will be set.

#rnx

diff --git a/Engine/Source/Runtime/D3D12RHI/Private/D3D12ExplicitDescriptorCache.cpp b/Engine/Source/Runtime/D3D12RHI/Private/D3D12ExplicitDescriptorCache.cpp
--- a/Engine/Source/Runtime/D3D12RHI/Private/D3D12ExplicitDescriptorCache.cpp
+++ b/Engine/Source/Runtime/D3D12RHI/Private/D3D12ExplicitDescriptorCache.cpp
@@ -246,8 +246,7 @@ int32 FD3D12ExplicitDescriptorHeap::Allocate(uint32 InNumDescriptors)
 			       TEXT("Explicit sampler descriptor heap overflow. ")
 			       TEXT("It is not possible to recover from this error, as maximum D3D12 sampler heap size is 2048."));
 		}
-		else if (GD3D12ExplicitViewDescriptorHeapSize == MaxNumDescriptors
-			&& FPlatformAtomics::InterlockedOr(&GD3D12ExplicitViewDescriptorHeapOverflowReported, 1) == 0)
+		else if ((uint32)GD3D12ExplicitViewDescriptorHeapSize <= MaxNumDescriptors && FPlatformAtomics::InterlockedOr(&GD3D12ExplicitViewDescriptorHeapOverflowReported, 1) == 0)
 		{
 			// NOTE: GD3D12RayTracingViewDescriptorHeapOverflowReported is set atomically because multiple 
 			// allocations may be happening simultaneously, but we only want to report the error once.
