Make sure the graphics command list is valid and open before trying to flush pending descriptor updates

diff --git a/Engine/Source/Runtime/D3D12RHI/Private/D3D12Commands.cpp b/Engine/Source/Runtime/D3D12RHI/Private/D3D12Commands.cpp
--- a/Engine/Source/Runtime/D3D12RHI/Private/D3D12Commands.cpp
+++ b/Engine/Source/Runtime/D3D12RHI/Private/D3D12Commands.cpp
@@ -77,6 +77,8 @@ bool FD3D12CommandContext::FlushPendingDescriptorUpdates()
 	bool bNewBindlessHeapsSet = false;
 
 #if PLATFORM_SUPPORTS_BINDLESS_RENDERING
+	// Make sure the graphics command list is valid and open before trying to flush pending descriptor updates
+	GraphicsCommandList();
 	bNewBindlessHeapsSet = GetParentDevice()->GetBindlessDescriptorManager().FlushPendingDescriptorUpdates(*this);
 #endif
 
