Move RHIEndFrame frame number field to the high level, so platform RHIs don't need to maintain their own counter.

diff --git a/Engine/Source/Runtime/D3D12RHI/Private/D3D12RHI.cpp b/Engine/Source/Runtime/D3D12RHI/Private/D3D12RHI.cpp
--- a/Engine/Source/Runtime/D3D12RHI/Private/D3D12RHI.cpp
+++ b/Engine/Source/Runtime/D3D12RHI/Private/D3D12RHI.cpp
@@ -421,7 +421,7 @@ void FD3D12DynamicRHI::FlushTiming(bool bCreateNew, const FRHIEndFrameArgs& Args
 		ProcessTimestamps(OldTiming);
 	};
 
-	EnqueueEndOfPipeTask(MoveTemp(Lambda), [&, Number = FrameNumber++](FD3D12Payload& Payload)
+	EnqueueEndOfPipeTask(MoveTemp(Lambda), [&](FD3D12Payload& Payload)
 	{
 		if (bCreateNew)
 		{
@@ -438,7 +438,7 @@ void FD3D12DynamicRHI::FlushTiming(bool bCreateNew, const FRHIEndFrameArgs& Args
 		if (bHasTiming && Payload.Queue.QueueType != ED3D12QueueType::Copy)
 		{
 			UE::RHI::GPUProfiler::FEvent::FFrameBoundary Boundary;
-			Boundary.FrameNumber = Number;
+			Boundary.FrameNumber = Args.FrameNumber;
 
 		#if WITH_RHI_BREADCRUMBS
 			switch (Payload.Queue.QueueType)
diff --git a/Engine/Source/Runtime/D3D12RHI/Private/D3D12RHIPrivate.h b/Engine/Source/Runtime/D3D12RHI/Private/D3D12RHIPrivate.h
--- a/Engine/Source/Runtime/D3D12RHI/Private/D3D12RHIPrivate.h
+++ b/Engine/Source/Runtime/D3D12RHI/Private/D3D12RHIPrivate.h
@@ -553,9 +553,6 @@ protected:
 	TRefCountPtr<IDXGIFactory2> DXGIFactoryForDisplayList;
 #endif
 
-	// Counts the number of calls to RHIEndFrame, and is used in GPU profiler frame boundary events.
-	uint32 FrameNumber = 0;
-
 public:
 
 	virtual FD3D12ResourceDesc GetResourceDesc(const FRHITextureDesc& CreateInfo) const;
diff --git a/Engine/Source/Runtime/RHI/Private/RHICommandList.cpp b/Engine/Source/Runtime/RHI/Private/RHICommandList.cpp
--- a/Engine/Source/Runtime/RHI/Private/RHICommandList.cpp
+++ b/Engine/Source/Runtime/RHI/Private/RHICommandList.cpp
@@ -1063,6 +1063,7 @@ void FRHICommandListExecutor::FSubmitState::Submit(const FSubmitArgs& Args)
 	{
 		FDynamicRHI::FRHIEndFrameArgs EndFrameArgs
 		{
+			.FrameNumber = GRHICommandList.FrameNumber++,
 #if WITH_RHI_BREADCRUMBS
 			.GPUBreadcrumbs = Args.GPUBreadcrumbs
 #endif
diff --git a/Engine/Source/Runtime/RHI/Public/DynamicRHI.h b/Engine/Source/Runtime/RHI/Public/DynamicRHI.h
--- a/Engine/Source/Runtime/RHI/Public/DynamicRHI.h
+++ b/Engine/Source/Runtime/RHI/Public/DynamicRHI.h
@@ -155,6 +155,9 @@ public:
 
 	struct FRHIEndFrameArgs
 	{
+		// Increments once per call to RHIEndFrame
+		uint32 FrameNumber;
+
 #if WITH_RHI_BREADCRUMBS
 		const TRHIPipelineArray<FRHIBreadcrumbNode*>& GPUBreadcrumbs;
 #endif
diff --git a/Engine/Source/Runtime/RHI/Public/RHICommandList.h b/Engine/Source/Runtime/RHI/Public/RHICommandList.h
--- a/Engine/Source/Runtime/RHI/Public/RHICommandList.h
+++ b/Engine/Source/Runtime/RHI/Public/RHICommandList.h
@@ -5375,6 +5375,9 @@ private:
 #endif
 	FRHIDrawStats FrameDrawStats;
 
+	// Counts the number of calls to RHIEndFrame, and is used in GPU profiler frame boundary events.
+	uint32 FrameNumber = 0;
+
 	bool AllowParallel() const;
 };
 
