Kick the D3D12 interrupt thread if any queue is processed, not just the last one. The previous check for the result didn't account for the fact that another queue could have processed work.

diff --git a/Engine/Source/Runtime/D3D12RHI/Private/D3D12Submission.cpp b/Engine/Source/Runtime/D3D12RHI/Private/D3D12Submission.cpp
--- a/Engine/Source/Runtime/D3D12RHI/Private/D3D12Submission.cpp
+++ b/Engine/Source/Runtime/D3D12RHI/Private/D3D12Submission.cpp
@@ -371,6 +371,7 @@ FD3D12DynamicRHI::FProcessResult FD3D12DynamicRHI::ProcessSubmissionQueue()
 	};
 
 	bool bProgress;
+	bool bKickInterruptThread = false;
 
 	do
 	{
@@ -445,6 +446,7 @@ FD3D12DynamicRHI::FProcessResult FD3D12DynamicRHI::ProcessSubmissionQueue()
 					CurrentQueue.PayloadToSubmit = Payload;
 					QueuesWithPayloads.Add(&CurrentQueue);
 					Result.Status |= EQueueStatus::Processed;
+					bKickInterruptThread = true;
 
 					//
 					// Now we generate any required barrier command lists. These may require
@@ -555,7 +557,7 @@ FD3D12DynamicRHI::FProcessResult FD3D12DynamicRHI::ProcessSubmissionQueue()
 
 	FlushPayloads();
 
-	if (InterruptThread && EnumHasAnyFlags(Result.Status, EQueueStatus::Processed))
+	if (InterruptThread && bKickInterruptThread)
 	{
 		InterruptThread->Kick();
 	}
