Remove unrequired check

diff --git a/Engine/Source/Runtime/D3D12RHI/Private/D3D12BindlessDescriptors.cpp b/Engine/Source/Runtime/D3D12RHI/Private/D3D12BindlessDescriptors.cpp
--- a/Engine/Source/Runtime/D3D12RHI/Private/D3D12BindlessDescriptors.cpp
+++ b/Engine/Source/Runtime/D3D12RHI/Private/D3D12BindlessDescriptors.cpp
@@ -383,8 +383,6 @@ void FD3D12BindlessResourceManager::UpdateDescriptor(FD3D12ContextArray const& C
 	if (DstHandle.IsValid())
 	{
 		TRACE_CPUPROFILER_EVENT_SCOPE(FD3D12BindlessResourceManager::UpdateDescriptor);
-
-		check(IsInRHIThread() || GRHICommandList.Bypass());
 	
 		FScopeLock ScopeLock(&HeapsCS);
 
