Change default buffer allocation alignment to 16 instead of 4 so Raw SRV views can be made on it (16 byte alignment requirement)

diff --git a/Engine/Source/Runtime/D3D12RHI/Private/D3D12Buffer.cpp b/Engine/Source/Runtime/D3D12RHI/Private/D3D12Buffer.cpp
--- a/Engine/Source/Runtime/D3D12RHI/Private/D3D12Buffer.cpp
+++ b/Engine/Source/Runtime/D3D12RHI/Private/D3D12Buffer.cpp
@@ -397,7 +397,7 @@ void FD3D12Buffer::GetResourceDescAndAlignment(uint64 InSize, uint32 InStride, E
 	else
 	{
 		// Structured buffers, non-ByteAddress buffers, need to be aligned to their stride to ensure that they can be addressed correctly with element based offsets.
-		Alignment = (InStride > 0) && (EnumHasAnyFlags(InUsage, BUF_StructuredBuffer) || !EnumHasAnyFlags(InUsage, BUF_ByteAddressBuffer | BUF_DrawIndirect)) ? InStride : 4;
+		Alignment = (InStride > 0) && (EnumHasAnyFlags(InUsage, BUF_StructuredBuffer) || !EnumHasAnyFlags(InUsage, BUF_ByteAddressBuffer | BUF_DrawIndirect)) ? InStride : 16;
 	}
 }
 
