Log shader hashes when we detect a missing uniform buffer in D3D12.

#rnx

diff --git a/Engine/Source/Runtime/D3D12RHI/Private/D3D12DescriptorCache.cpp b/Engine/Source/Runtime/D3D12RHI/Private/D3D12DescriptorCache.cpp
--- a/Engine/Source/Runtime/D3D12RHI/Private/D3D12DescriptorCache.cpp
+++ b/Engine/Source/Runtime/D3D12RHI/Private/D3D12DescriptorCache.cpp
@@ -7,6 +7,8 @@
 #include "D3D12ExplicitDescriptorCache.h"
 #include "D3D12RHIPrivate.h"
 #include "D3D12RayTracing.h"
+#include "D3D12StateCachePrivate.h"
+#include "D3D12PipelineState.h"
 
 bool FD3D12DescriptorCache::HeapRolledOver(ERHIDescriptorHeapType InHeapType)
 {
@@ -649,7 +651,7 @@ void FD3D12DescriptorCache::SetConstantBufferViews(EShaderFrequency ShaderStage,
 #endif // D3D12RHI_USE_CONSTANT_BUFFER_VIEWS
 }
 
-void FD3D12DescriptorCache::SetRootConstantBuffers(EShaderFrequency ShaderStage, const FD3D12RootSignature* RootSignature, FD3D12ConstantBufferCache& Cache, CBVSlotMask SlotsNeededMask)
+void FD3D12DescriptorCache::SetRootConstantBuffers(EShaderFrequency ShaderStage, const FD3D12RootSignature* RootSignature, FD3D12ConstantBufferCache& Cache, CBVSlotMask SlotsNeededMask, FD3D12StateCache* StateCache)
 {
 	CBVSlotMask& CurrentDirtySlotMask = Cache.DirtySlotMask[ShaderStage];
 	check(CurrentDirtySlotMask != 0);	// All dirty slots for the current shader stage.
@@ -673,7 +675,74 @@ void FD3D12DescriptorCache::SetRootConstantBuffers(EShaderFrequency ShaderStage,
 			const D3D12_GPU_VIRTUAL_ADDRESS CurrentGPUVirtualAddress = Cache.CurrentGPUVirtualAddress[ShaderStage][SlotIndex];
 			if (CurrentGPUVirtualAddress == 0)
 			{
-				UE_LOG(LogD3D12RHI, Fatal, TEXT("Missing uniform buffer at slot %u, stage %s. Please check the high level drawing code."), SlotIndex, GetShaderFrequencyString(ShaderStage));
+				FString ShaderHashList;
+
+				switch(ShaderStage)
+				{
+					case SF_Vertex:
+					case SF_Mesh:
+					case SF_Amplification:
+					case SF_Pixel:
+					case SF_Geometry:
+					{
+						FD3D12GraphicsPipelineState* GraphicsPSO = StateCache->GetGraphicsPipelineState();
+						if (!GraphicsPSO)
+						{
+							// Shouldn't happen, but we don't want to crash while crashing.
+							ShaderHashList = TEXT("NO GRAPHICS PSO!");
+							break;
+						}
+
+						FSHA1 PipelineHasher;
+
+						const auto AddShaderHash = [&PipelineHasher, &ShaderHashList](const FRHIShader* Shader)
+						{
+							FSHAHash ShaderHash;
+							if (Shader)
+							{
+								ShaderHash = Shader->GetHash();
+								ShaderHashList.Appendf(TEXT("%s: %s, "), GetShaderFrequencyString(Shader->GetFrequency(), false), *ShaderHash.ToString());
+							}
+							PipelineHasher.Update(&ShaderHash.Hash[0], sizeof(FSHAHash));
+						};
+
+						AddShaderHash(GraphicsPSO->GetVertexShader());
+						AddShaderHash(GraphicsPSO->GetMeshShader());
+						AddShaderHash(GraphicsPSO->GetAmplificationShader());
+						AddShaderHash(GraphicsPSO->GetPixelShader());
+						AddShaderHash(GraphicsPSO->GetGeometryShader());
+
+						PipelineHasher.Final();
+						FSHAHash PipelineHash;
+						PipelineHasher.GetHash(&PipelineHash.Hash[0]);
+
+						ShaderHashList.Appendf(TEXT("Pipeline: %s"), *PipelineHash.ToString());
+						break;
+					}
+
+					case SF_Compute:
+					{
+						FD3D12ComputePipelineState* ComputePSO = StateCache->GetComputePipelineState();
+						if (ComputePSO && ComputePSO->GetComputeShader())
+						{
+							ShaderHashList.Appendf(TEXT("Compute: %s"), *ComputePSO->GetComputeShader()->GetHash().ToString());
+						}
+						else
+						{
+							// Shouldn't happen, but we don't want to crash while crashing.
+							ShaderHashList = TEXT("NO COMPUTE SHADER!");
+						}
+						break;
+					}
+
+					default:
+					{
+						ShaderHashList = TEXT("NO PSO FOR STAGE!");
+						break;
+					}
+				}
+
+				UE_LOG(LogD3D12RHI, Fatal, TEXT("Missing uniform buffer at slot %u, stage %s. Please check the high level drawing code. Hashes: %s."), SlotIndex, GetShaderFrequencyString(ShaderStage), *ShaderHashList);
 			}
 
 			if (ShaderStage == SF_Compute)
diff --git a/Engine/Source/Runtime/D3D12RHI/Private/D3D12DescriptorCache.h b/Engine/Source/Runtime/D3D12RHI/Private/D3D12DescriptorCache.h
--- a/Engine/Source/Runtime/D3D12RHI/Private/D3D12DescriptorCache.h
+++ b/Engine/Source/Runtime/D3D12RHI/Private/D3D12DescriptorCache.h
@@ -26,6 +26,7 @@ struct FD3D12ConstantBufferCache;
 struct FD3D12ShaderResourceViewCache;
 struct FD3D12UnorderedAccessViewCache;
 struct FD3D12SamplerStateCache;
+class FD3D12StateCache;
 
 class FD3D12SyncPoint;
 using FD3D12SyncPointRef = TRefCountPtr<FD3D12SyncPoint>;
@@ -345,7 +346,7 @@ public:
 	void SetSRVTable(EShaderFrequency ShaderStage, const FD3D12RootSignature* RootSignature, FD3D12ShaderResourceViewCache& Cache, uint32 SlotsNeeded, const D3D12_GPU_DESCRIPTOR_HANDLE& BindDescriptor);
 
 	void SetConstantBufferViews(EShaderFrequency ShaderStage, const FD3D12RootSignature* RootSignature, FD3D12ConstantBufferCache& Cache, CBVSlotMask SlotsNeededMask, uint32 Count, uint32& HeapSlot);
-	void SetRootConstantBuffers(EShaderFrequency ShaderStage, const FD3D12RootSignature* RootSignature, FD3D12ConstantBufferCache& Cache, CBVSlotMask SlotsNeededMask);
+	void SetRootConstantBuffers(EShaderFrequency ShaderStage, const FD3D12RootSignature* RootSignature, FD3D12ConstantBufferCache& Cache, CBVSlotMask SlotsNeededMask, FD3D12StateCache* StateCache);
 
 	void PrepareBindlessViews(EShaderFrequency ShaderStage, TConstArrayView<FD3D12ShaderResourceView*> SRVs, TConstArrayView<FD3D12UnorderedAccessView*> UAVs);
 
diff --git a/Engine/Source/Runtime/D3D12RHI/Private/D3D12StateCache.cpp b/Engine/Source/Runtime/D3D12RHI/Private/D3D12StateCache.cpp
--- a/Engine/Source/Runtime/D3D12RHI/Private/D3D12StateCache.cpp
+++ b/Engine/Source/Runtime/D3D12RHI/Private/D3D12StateCache.cpp
@@ -886,7 +886,7 @@ void FD3D12StateCache::ApplyConstants(const FD3D12RootSignature* const pRootSign
 		{
 			if (CurrentShaderDirtyCBVSlots[Index])
 			{
-				DescriptorCache.SetRootConstantBuffers(static_cast<EShaderFrequency>(Index), pRootSignature, CBVCache, CurrentShaderDirtyCBVSlots[Index]);
+				DescriptorCache.SetRootConstantBuffers(static_cast<EShaderFrequency>(Index), pRootSignature, CBVCache, CurrentShaderDirtyCBVSlots[Index], this);
 			}
 		}
 	}
