Add D3D12_MESSAGE_ID_HEAP_ADDRESS_RANGE_HAS_NO_RESOURCE to deny error list because of Agility SDK bug (version 614) - should be fixed in the next Agility SDK according to MS

diff --git a/Engine/Source/Runtime/D3D12RHI/Private/D3D12Adapter.cpp b/Engine/Source/Runtime/D3D12RHI/Private/D3D12Adapter.cpp
--- a/Engine/Source/Runtime/D3D12RHI/Private/D3D12Adapter.cpp
+++ b/Engine/Source/Runtime/D3D12RHI/Private/D3D12Adapter.cpp
@@ -730,6 +730,11 @@ void FD3D12Adapter::CreateRootDevice(bool bWithDebug)
 				// This is safe to ignore
 				D3D12_MESSAGE_ID_META_COMMAND_UNSUPPORTED_PARAMS,
 
+				// D3D Agility SDK bug (version 614) where scratch allocation memory for BuildRaytracingAccelerationStructure is always computed using build size
+				// while operation is update and needs less memory (verified by MS and will be fixed in the next Agility SDK)
+				// See: UE-222685 to remove again from Deny list after Agility SDK upgrade
+				D3D12_MESSAGE_ID_HEAP_ADDRESS_RANGE_HAS_NO_RESOURCE,
+
 			};
 
 #if PLATFORM_DESKTOP
