Add a CVar to disable D3D runtime's background threads (r.D3D12.DevDisableD3DRuntimeBackgroundThreads).

This is useful for A/B profiling and other testing that can be susceptible to having threads spawned and actively used outside of the engine's control. The background threads are commonly used by drivers to schedule background PSO re-compilation (optimization) work without blocking the application. Note that this can change the performance characteristics of PSO creation as the driver may decide to perform a more expensive PSO compilation right away on the PSO creation thread.

Disabling background threads is only available when Windows developer mode is enabled.

#tests tested locally

diff --git a/Engine/Source/Runtime/D3D12RHI/Private/Windows/WindowsD3D12Device.cpp b/Engine/Source/Runtime/D3D12RHI/Private/Windows/WindowsD3D12Device.cpp
--- a/Engine/Source/Runtime/D3D12RHI/Private/Windows/WindowsD3D12Device.cpp
+++ b/Engine/Source/Runtime/D3D12RHI/Private/Windows/WindowsD3D12Device.cpp
@@ -76,6 +76,15 @@ static TAutoConsoleVariable<int32> CVarExperimentalShaderModels(
 );
 #endif // !UE_BUILD_SHIPPING
 
+// See https://microsoft.github.io/DirectX-Specs/d3d/BackgroundProcessing.html.
+int32 GDevDisableD3DRuntimeBackgroundThreads = 0;
+static FAutoConsoleVariableRef CVarDevDisableD3DRuntimeBackgroundThreads(
+	TEXT("r.D3D12.DevDisableD3DRuntimeBackgroundThreads"),
+	GDevDisableD3DRuntimeBackgroundThreads,
+	TEXT("If > 0, disables the background threads created by the D3D runtime for background shader optimization. Only available when Windows developer mode is enabled. (default = 0)."),
+	ECVF_ReadOnly
+);
+
 #if D3D12RHI_SUPPORTS_WIN_PIX
 int32 GAutoAttachPIX = 0;
 static FAutoConsoleVariableRef CVarAutoAttachPIX(
@@ -1371,6 +1380,33 @@ void FD3D12DynamicRHI::Init()
 		Adapter->InitializeDevices();
 	}
 
+	if (GDevDisableD3DRuntimeBackgroundThreads)
+	{
+#if D3D12_MAX_DEVICE_INTERFACE >= 6
+		ID3D12Device6* Device6 = GetAdapter().GetD3DDevice6();
+		if (Device6)
+		{
+			HRESULT Res = Device6->SetBackgroundProcessingMode(
+				D3D12_BACKGROUND_PROCESSING_MODE_DISABLE_PROFILING_BY_SYSTEM,
+				D3D12_MEASUREMENTS_ACTION_KEEP_ALL,
+				nullptr, nullptr);
+
+			if (SUCCEEDED(Res))
+			{
+				UE_LOG(LogD3D12RHI, Log, TEXT("Disabled D3D runtime's background threads"));
+			}
+			else
+			{
+				UE_LOG(LogD3D12RHI, Error, TEXT("Could not disable D3D runtime's background threads: SetBackgroundProcessingMode returned error 0x%08X"), Res);
+			}
+		}
+		else
+#endif  
+		{
+			UE_LOG(LogD3D12RHI, Warning, TEXT("Could not disable D3D runtime's background threads because the ID3D12Device6 interface is not available"));
+		}
+	}
+
 	bool bHasVendorSupportForAtomic64 = false;
 
 #if WITH_AMD_AGS
