Static analysis fix - uninitialized pointer

diff --git a/Engine/Source/Runtime/D3D12RHI/Private/D3D12Allocation.cpp b/Engine/Source/Runtime/D3D12RHI/Private/D3D12Allocation.cpp
--- a/Engine/Source/Runtime/D3D12RHI/Private/D3D12Allocation.cpp
+++ b/Engine/Source/Runtime/D3D12RHI/Private/D3D12Allocation.cpp
@@ -2345,7 +2345,7 @@ FD3D12SegHeap* FD3D12SegList::CreateBackingHeap(
 {
 	// CS can be unlocked at this point and re-locked before adding it to FreeHeaps
 	// but doing so may cause multiple heaps to be created
-	ID3D12Heap* D3DHeap;
+	ID3D12Heap* D3DHeap = nullptr;
 	D3D12_HEAP_DESC Desc = {};
 	Desc.SizeInBytes = HeapSize;
 	Desc.Properties = CD3DX12_HEAP_PROPERTIES(HeapType, Parent->GetGPUMask().GetNative(), VisibleNodeMask.GetNative());
