Static analysis fix - Maximum size of a root signature is 64 DWORDs, print an error message if trying to set over 64 indexes
https://learn.microsoft.com/en-us/windows/win32/direct3d12/root-signature-limits

diff --git a/Engine/Source/Runtime/D3D12RHI/Private/D3D12DescriptorCache.cpp b/Engine/Source/Runtime/D3D12RHI/Private/D3D12DescriptorCache.cpp
--- a/Engine/Source/Runtime/D3D12RHI/Private/D3D12DescriptorCache.cpp
+++ b/Engine/Source/Runtime/D3D12RHI/Private/D3D12DescriptorCache.cpp
@@ -629,13 +629,29 @@ void FD3D12DescriptorCache::SetConstantBufferViews(EShaderFrequency ShaderStage,
 	{
 		const uint32 RDTIndex = RootSignature->CBVRDTBindSlot(ShaderStage);
 		ensure(RDTIndex != 255);
-		Context.GraphicsCommandList()->SetComputeRootDescriptorTable(RDTIndex, BindDescriptor);
+
+		if (RDTIndex < 64)
+		{
+			Context.GraphicsCommandList()->SetComputeRootDescriptorTable(RDTIndex, BindDescriptor);
+		}
+		else
+		{
+			UE_LOG(LogD3D12RHI, Error, TEXT("SetComputeRootDescriptorTable call exceed max 64 slots: %d "), RDTIndex);
+		}
 	}
 	else
 	{
 		const uint32 RDTIndex = RootSignature->CBVRDTBindSlot(ShaderStage);
 		ensure(RDTIndex != 255);
-		Context.GraphicsCommandList()->SetGraphicsRootDescriptorTable(RDTIndex, BindDescriptor);
+
+		if (RDTIndex < 64)
+		{
+			Context.GraphicsCommandList()->SetGraphicsRootDescriptorTable(RDTIndex, BindDescriptor);
+		}
+		else
+		{
+			UE_LOG(LogD3D12RHI, Error, TEXT("SetGraphicsRootDescriptorTable call exceed max 64 slots: %d "), RDTIndex);
+		}
 	}
 
 	// We changed the descriptor table, so all resources bound to slots outside of the table's range are now dirty.
@@ -745,13 +761,21 @@ void FD3D12DescriptorCache::SetRootConstantBuffers(EShaderFrequency ShaderStage,
 				UE_LOG(LogD3D12RHI, Fatal, TEXT("Missing uniform buffer at slot %u, stage %s. Please check the high level drawing code. Hashes: %s."), SlotIndex, GetShaderFrequencyString(ShaderStage), *ShaderHashList);
 			}
 
-			if (ShaderStage == SF_Compute)
+
+			if ((BaseIndex + SlotIndex) < 64)
 			{
-				Context.GraphicsCommandList()->SetComputeRootConstantBufferView(BaseIndex + SlotIndex, CurrentGPUVirtualAddress);
+				if (ShaderStage == SF_Compute)
+				{
+					Context.GraphicsCommandList()->SetComputeRootConstantBufferView(BaseIndex + SlotIndex, CurrentGPUVirtualAddress);
+				}
+				else
+				{
+					Context.GraphicsCommandList()->SetGraphicsRootConstantBufferView(BaseIndex + SlotIndex, CurrentGPUVirtualAddress);
+				}
 			}
 			else
 			{
-				Context.GraphicsCommandList()->SetGraphicsRootConstantBufferView(BaseIndex + SlotIndex, CurrentGPUVirtualAddress);
+				UE_LOG(LogD3D12RHI, Error, TEXT("%s call exceed max 64 slots: %d "), (ShaderStage == SF_Compute) ? TEXT("SetComputeRootConstantBufferView") : TEXT("SetGraphicsRootConstantBufferView"), BaseIndex + SlotIndex);
 			}
 
 			// Update residency.
