MGPU Path Tracer GPU hang fix:  Bindless descriptor heap allocator is now shared across GPUs, making it so a given resource has the same descriptor handle regardless of GPU.  Solves issues with platform agnostic logic operating on descriptor handles -- specifically incorrect FRHIShaderBindingTable's on the second GPU were causing the hang.  The descriptor allocators were moved from the device to the adapter, and descriptor handles are allocated by the first D3D12 linked object during initialization loops for samplers and resources, then shared by subsequent objects.  The first linked object is then responsible for freeing the shared handle.

#rnx

diff --git a/Engine/Source/Runtime/D3D12RHI/Private/D3D12Adapter.cpp b/Engine/Source/Runtime/D3D12RHI/Private/D3D12Adapter.cpp
--- a/Engine/Source/Runtime/D3D12RHI/Private/D3D12Adapter.cpp
+++ b/Engine/Source/Runtime/D3D12RHI/Private/D3D12Adapter.cpp
@@ -1187,6 +1187,16 @@ void FD3D12Adapter::InitializeDevices()
 		bTrackAllAllocation = (GD3D12TrackAllAlocations || UE::RHI::UseGPUCrashDebugging() || bTraceMemAlloc) && (GetResourceHeapTier() == D3D12_RESOURCE_HEAP_TIER_2);
 #endif 
 
+#if PLATFORM_SUPPORTS_BINDLESS_RENDERING
+		if (Desc.MaxSupportedFeatureLevel >= D3D_FEATURE_LEVEL_12_0 && Desc.MaxSupportedShaderModel >= D3D_SHADER_MODEL_6_6 && Desc.ResourceBindingTier >= D3D12_RESOURCE_BINDING_TIER_3)
+		{
+			// Needs to happen before device creation below
+			BindlessManager.Init(this);
+
+			GRHIGlobals.ShaderBundles.RequiresSharedBindlessParameters = (BindlessManager.GetResourcesConfiguration() == ERHIBindlessConfiguration::AllShaders || BindlessManager.GetSamplersConfiguration() == ERHIBindlessConfiguration::AllShaders);
+		}
+#endif
+
 		// Context redirectors allow RHI commands to be executed on multiple GPUs at the
 		// same time in a multi-GPU system. Redirectors have a physical mask for the GPUs
 		// they can support and an active mask which restricts commands to operate on a
@@ -1230,19 +1240,6 @@ void FD3D12Adapter::InitializeDevices()
 		ERHIBindlessConfiguration BindlessResourcesConfig = ERHIBindlessConfiguration::Disabled;
 		ERHIBindlessConfiguration BindlessSamplersConfig = ERHIBindlessConfiguration::Disabled;
 
-#if PLATFORM_SUPPORTS_BINDLESS_RENDERING
-		if (Desc.MaxSupportedFeatureLevel >= D3D_FEATURE_LEVEL_12_0 && Desc.MaxSupportedShaderModel >= D3D_SHADER_MODEL_6_6 && Desc.ResourceBindingTier >= D3D12_RESOURCE_BINDING_TIER_3)
-		{
-			BindlessResourcesConfig = RHIGetRuntimeBindlessResourcesConfiguration(GMaxRHIShaderPlatform);
-			BindlessSamplersConfig = RHIGetRuntimeBindlessSamplersConfiguration(GMaxRHIShaderPlatform);
-
-			bBindlessResourcesAllowed = (BindlessResourcesConfig != ERHIBindlessConfiguration::Disabled);
-			bBindlessSamplersAllowed = (BindlessSamplersConfig != ERHIBindlessConfiguration::Disabled);
-
-			GRHIGlobals.ShaderBundles.RequiresSharedBindlessParameters = (BindlessResourcesConfig == ERHIBindlessConfiguration::AllShaders || BindlessSamplersConfig == ERHIBindlessConfiguration::AllShaders);
-		}
-#endif
-
 #if USE_STATIC_ROOT_SIGNATURE
 		EShaderBindingLayoutFlags GraphicsFlags{};
 #if PLATFORM_SUPPORTS_MESH_SHADERS
diff --git a/Engine/Source/Runtime/D3D12RHI/Private/D3D12Adapter.h b/Engine/Source/Runtime/D3D12RHI/Private/D3D12Adapter.h
--- a/Engine/Source/Runtime/D3D12RHI/Private/D3D12Adapter.h
+++ b/Engine/Source/Runtime/D3D12RHI/Private/D3D12Adapter.h
@@ -36,6 +36,7 @@
 #include "D3D12ThirdParty.h"
 #include "D3D12CommandContext.h"
 #include "D3D12RootSignature.h"
+#include "D3D12BindlessDescriptors.h"
 
 class FD3D12TransientHeapCache;
 class IRHITransientMemoryCache;
@@ -448,6 +449,8 @@ public:
 	HMODULE GetDxgiDllHandle() const { return DxgiDllHandle; };
 #endif
 
+	FD3D12BindlessManagerAdapter& GetBindlessManager() { return BindlessManager; }
+
 protected:
 
 	virtual void CreateRootDevice(bool bWithDebug);
@@ -548,11 +551,6 @@ protected:
 
 	FD3D12AdapterDesc Desc;
 
-#if PLATFORM_SUPPORTS_BINDLESS_RENDERING
-	bool bBindlessResourcesAllowed = false;
-	bool bBindlessSamplersAllowed = false;
-#endif
-
 	TRefCountPtr<IDXGIAdapter> DxgiAdapter;
 
 	FD3D12RootSignatureManager RootSignatureManager;
@@ -622,6 +620,8 @@ protected:
 	FD3D12RootSignature StaticRayTracingLocalRootSignature;
 #endif
 
+	FD3D12BindlessManagerAdapter BindlessManager;
+
 private:
 	// Insight memory trace helper
 	void TraceMemoryAllocation(FD3D12Resource* Resource);
diff --git a/Engine/Source/Runtime/D3D12RHI/Private/D3D12BindlessDescriptors.cpp b/Engine/Source/Runtime/D3D12RHI/Private/D3D12BindlessDescriptors.cpp
--- a/Engine/Source/Runtime/D3D12RHI/Private/D3D12BindlessDescriptors.cpp
+++ b/Engine/Source/Runtime/D3D12RHI/Private/D3D12BindlessDescriptors.cpp
@@ -76,23 +76,11 @@ void FD3D12BindlessSamplerManager::CleanupResources()
 	GpuHeap = nullptr;
 }
 
-FRHIDescriptorHandle FD3D12BindlessSamplerManager::AllocateAndInitialize(FD3D12SamplerState* SamplerState)
+void FD3D12BindlessSamplerManager::InitializeSampler(FRHIDescriptorHandle DstHandle, FD3D12SamplerState* SamplerState)
 {
-	FRHIDescriptorHandle Result = Allocator.Allocate();
-	if (ensure(Result.IsValid()))
-	{
-		UE::D3D12Descriptors::CopyDescriptor(GetParentDevice(), GpuHeap, Result, SamplerState->OfflineDescriptor);
-	}
-	check(Result.IsValid());
-	return Result;
-}
+	check(DstHandle.GetType() == ERHIDescriptorHeapType::Sampler);
 
-void FD3D12BindlessSamplerManager::Free(FRHIDescriptorHandle InHandle)
-{
-	if (InHandle.IsValid())
-	{
-		Allocator.Free(InHandle);
-	}
+	UE::D3D12Descriptors::CopyDescriptor(GetParentDevice(), GpuHeap, DstHandle, SamplerState->OfflineDescriptor);
 }
 
 void FD3D12BindlessSamplerManager::OpenCommandList(FD3D12CommandContext& Context)
@@ -117,62 +105,138 @@ FD3D12DescriptorHeap* FD3D12BindlessSamplerManager::GetExplicitHeapForContext(FD
 }
 
 ////////////////////////////////////////////////////////////////////////////////////////////////////////////////////////
-// FD3D12BindlessResourceManager
+// FD3D12BindlessManagerAdapter
 
-#if !D3D12RHI_CUSTOM_BINDLESS_RESOURCE_MANAGER
+void FD3D12BindlessManagerAdapter::Init(FD3D12Adapter* InParent)
+{
+	SetParentAdapter(InParent);
 
-FD3D12BindlessResourceManager::FD3D12BindlessResourceManager(FD3D12Device* InDevice, ERHIBindlessConfiguration InConfiguration, uint32 InNumDescriptors, TConstArrayView<TStatId> InStats)
-	: FD3D12DeviceChild(InDevice)
-	, CpuHeap(UE::D3D12BindlessDescriptors::CreateCpuHeap(InDevice, ERHIDescriptorHeapType::Standard, InNumDescriptors))
-	, Allocator(ERHIDescriptorHeapType::Standard, InNumDescriptors, InStats)
-	, Configuration(InConfiguration)
-{	
-	if (GetConfiguration() == ERHIBindlessConfiguration::AllShaders)
+	BindlessResourcesConfiguration = RHIGetRuntimeBindlessResourcesConfiguration(GMaxRHIShaderPlatform);
+	BindlessSamplersConfiguration = RHIGetRuntimeBindlessSamplersConfiguration(GMaxRHIShaderPlatform);
+
+	if (BindlessResourcesConfiguration != ERHIBindlessConfiguration::Disabled)
 	{
-		// Always allocate a heap when full bindless
-		ActiveGpuHeapIndex = AddActiveGPUHeap();
+		const TStatId Stats[] =
+		{
+			GET_STATID(STAT_ResourceDescriptorsAllocated),
+			GET_STATID(STAT_BindlessResourceDescriptorsAllocated),
+		};
+
+		uint32 NumResourceDescriptors = GBindlessResourceDescriptorHeapSize;
+#if D3D12RHI_USE_CONSTANT_BUFFER_VIEWS
+		NumResourceDescriptors += GBindlessOnlineDescriptorHeapBlockSize;
+#endif
+
+		ResourceAllocator = new FRHIHeapDescriptorAllocator(ERHIDescriptorHeapType::Standard, NumResourceDescriptors, Stats);
+	}
+
+	if (BindlessSamplersConfiguration != ERHIBindlessConfiguration::Disabled)
+	{
+		const TStatId Stats[] =
+		{
+			GET_STATID(STAT_SamplerDescriptorsAllocated),
+			GET_STATID(STAT_BindlessSamplerDescriptorsAllocated),
+		};
+
+		uint32 NumSamplerDescriptors = GBindlessSamplerDescriptorHeapSize;
+		if (NumSamplerDescriptors > D3D12_MAX_SHADER_VISIBLE_SAMPLER_HEAP_SIZE)
+		{
+			UE_LOG(LogD3D12RHI, Error, TEXT("D3D12.Bindless.SamplerDescriptorHeapSize was set to %d, which is higher than the D3D12 maximum of %d. Adjusting the value to prevent a crash."),
+				NumSamplerDescriptors,
+				D3D12_MAX_SHADER_VISIBLE_SAMPLER_HEAP_SIZE
+			);
+			NumSamplerDescriptors = D3D12_MAX_SHADER_VISIBLE_SAMPLER_HEAP_SIZE;
+		}
+
+		SamplerAllocator = new FRHIHeapDescriptorAllocator(ERHIDescriptorHeapType::Sampler, NumSamplerDescriptors, Stats);
+	}
+
+	bBindlessResourcesAllowed = (BindlessResourcesConfiguration != ERHIBindlessConfiguration::Disabled);
+	bBindlessSamplersAllowed = (BindlessSamplersConfiguration != ERHIBindlessConfiguration::Disabled);
+}
+
+FRHIDescriptorHandle FD3D12BindlessManagerAdapter::AllocateSamplerHandle()
+{
+	FRHIDescriptorHandle Result = SamplerAllocator->Allocate();
+	check(Result.IsValid());
+	return Result;
+}
+
+void FD3D12BindlessManagerAdapter::FreeSamplerHandle(FRHIDescriptorHandle InHandle)
+{
+	if (InHandle.IsValid())
+	{
+		SamplerAllocator->Free(InHandle);
 	}
 }
 
-FRHIDescriptorHandle FD3D12BindlessResourceManager::Allocate()
+FRHIDescriptorHandle FD3D12BindlessManagerAdapter::AllocateResourceHandle()
 {
-	FRHIDescriptorHandle Result = Allocator.Allocate();
+	if (!bBindlessResourcesAllowed)
+	{
+		return FRHIDescriptorHandle();
+	}
+
+	FRHIDescriptorHandle Result = ResourceAllocator->Allocate();
 	if (!Result.IsValid())
 	{
-		Result = ResizeGrowAndAllocate();
+		TRACE_CPUPROFILER_EVENT_SCOPE(FD3D12Adapter::BindlessResourceAllocateHandle(GrowHeap));
+
+		FScopeLock ScopeLock(&ResourceHeapsCS);
+
+		// Grow the descriptor handle allocator
+		uint32 CurrentNumDescriptors = ResourceAllocator->GetCapacity();
+		uint32 NewNumDescriptors = CurrentNumDescriptors * 2;
+		Result = ResourceAllocator->ResizeGrowAndAllocate(NewNumDescriptors, ResourceAllocator->GetType());
+
+		// Grow the CPU heaps for all devices
+		for (FD3D12Device* ParentDevice : GetParentAdapter()->GetDevices())
+		{
+			ParentDevice->GetBindlessDescriptorManager().GetResourceManager().GrowCPUHeap(CurrentNumDescriptors, NewNumDescriptors);
+		}
+
+		return Result;
 	}
 	check(Result.IsValid());
 	return Result;
 }
 
-void FD3D12BindlessResourceManager::Free(FRHIDescriptorHandle InHandle)
+void FD3D12BindlessManagerAdapter::FreeResourceHandle(FRHIDescriptorHandle InHandle)
 {
 	if (InHandle.IsValid())
 	{
-		Allocator.Free(InHandle);
+		ResourceAllocator->Free(InHandle);
 	}
 }
 
-FRHIDescriptorHandle FD3D12BindlessResourceManager::ResizeGrowAndAllocate()
-{	
-	TRACE_CPUPROFILER_EVENT_SCOPE(FD3D12BindlessResourceManager::ResizeGrow);
 
-	FScopeLock ScopeLock(&HeapsCS);
+////////////////////////////////////////////////////////////////////////////////////////////////////////////////////////
+// FD3D12BindlessResourceManager
 
-	// Grow the descriptor handle allocator
-	uint32 CurrentNumDescriptors = Allocator.GetCapacity();
-	uint32 NewNumDescriptors = CurrentNumDescriptors * 2;
-	FRHIDescriptorHandle Result = Allocator.ResizeGrowAndAllocate(NewNumDescriptors, Allocator.GetType());
+#if !D3D12RHI_CUSTOM_BINDLESS_RESOURCE_MANAGER
+
+FD3D12BindlessResourceManager::FD3D12BindlessResourceManager(FD3D12Device* InDevice, FD3D12Adapter* InAdapter)
+	: FD3D12DeviceChild(InDevice)
+	, HeapsCS(InAdapter->GetBindlessManager().GetResourceHeapsCS())
+	, CpuHeap(UE::D3D12BindlessDescriptors::CreateCpuHeap(InDevice, ERHIDescriptorHeapType::Standard, InAdapter->GetBindlessManager().GetResourceCapacity()))
+	, Configuration(InAdapter->GetBindlessManager().GetResourcesConfiguration())
+{	
+	if (GetConfiguration() == ERHIBindlessConfiguration::AllShaders)
+	{
+		// Always allocate a heap when full bindless
+		ActiveGpuHeapIndex = AddActiveGPUHeap();
+	}
+}
 
+void FD3D12BindlessResourceManager::GrowCPUHeap(uint32 OriginalNumDescriptors, uint32 NewNumDescriptors)
+{
 	// Allocate new cpu heap & copy over the content
 	FD3D12DescriptorHeapPtr NewCpuHeap = UE::D3D12BindlessDescriptors::CreateCpuHeap(GetParentDevice(), ERHIDescriptorHeapType::Standard, NewNumDescriptors);
-	UE::D3D12Descriptors::CopyDescriptors(GetParentDevice(), NewCpuHeap, CpuHeap, 0, CurrentNumDescriptors);
+	UE::D3D12Descriptors::CopyDescriptors(GetParentDevice(), NewCpuHeap, CpuHeap, 0, OriginalNumDescriptors);
 	CpuHeap = NewCpuHeap;
 
 	bRequestNewActiveGpuHeap = true;
 	bCPUHeapResized = true;
-
-	return Result;
 }
 
 void FD3D12BindlessResourceManager::CleanupResources()
@@ -229,7 +293,7 @@ int FD3D12BindlessResourceManager::AddActiveGPUHeap()
 	}
 	else
 	{
-		GpuHeapData.GpuHeap = UE::D3D12BindlessDescriptors::CreateGpuHeap(GetParentDevice(), Allocator.GetType(), Allocator.GetCapacity());
+		GpuHeapData.GpuHeap = UE::D3D12BindlessDescriptors::CreateGpuHeap(GetParentDevice(), CpuHeap->GetType(), CpuHeap->GetNumDescriptors());
 
 		INC_DWORD_STAT(STAT_D3D12BindlessResourceHeapsAllocated);
 		INC_MEMORY_STAT_BY(STAT_D3D12BindlessResourceHeapGPUMemoryUsage, GpuHeapData.GpuHeap->GetMemorySize());
@@ -509,7 +573,7 @@ void FD3D12BindlessResourceManager::CopyCpuHeap(FD3D12DescriptorHeap* Destinatio
 {
 	// Copy the smallest possible set of descriptors from the CPU heap to the new GPU heap.
 	FRHIDescriptorAllocatorRange AllocatedRange(0, 0);
-	if (Allocator.GetAllocatedRange(AllocatedRange))
+	if (GetParentDevice()->GetParentAdapter()->GetBindlessManager().GetResourceAllocatedRange(AllocatedRange))
 	{
 		const uint32 NumDescriptorsToCopy = AllocatedRange.Last - AllocatedRange.First + 1;
 		UE::D3D12Descriptors::CopyDescriptors(GetParentDevice(), DestinationHeap, CpuHeap, AllocatedRange.First, NumDescriptorsToCopy);
@@ -646,23 +710,13 @@ FD3D12BindlessDescriptorManager::~FD3D12BindlessDescriptorManager() = default;
 
 void FD3D12BindlessDescriptorManager::Init()
 {
-	ResourcesConfiguration = RHIGetRuntimeBindlessResourcesConfiguration(GMaxRHIShaderPlatform);
-	SamplersConfiguration  = RHIGetRuntimeBindlessSamplersConfiguration(GMaxRHIShaderPlatform);
+	FD3D12Adapter* Adapter = GetParentDevice()->GetParentAdapter();
+	ResourcesConfiguration = Adapter->GetBindlessManager().GetResourcesConfiguration();
+	SamplersConfiguration  = Adapter->GetBindlessManager().GetSamplersConfiguration();
 
 	if (ResourcesConfiguration != ERHIBindlessConfiguration::Disabled)
 	{
-		const TStatId Stats[] =
-		{
-			GET_STATID(STAT_ResourceDescriptorsAllocated),
-			GET_STATID(STAT_BindlessResourceDescriptorsAllocated),
-		};
-
-		uint32 NumResourceDescriptors = GBindlessResourceDescriptorHeapSize;
-#if D3D12RHI_USE_CONSTANT_BUFFER_VIEWS
-		NumResourceDescriptors += GBindlessOnlineDescriptorHeapBlockSize;
-#endif
-
-		ResourceManager = MakeUnique<FD3D12BindlessResourceManager>(GetParentDevice(), ResourcesConfiguration, NumResourceDescriptors, Stats);
+		ResourceManager = MakeUnique<FD3D12BindlessResourceManager>(GetParentDevice(), Adapter);
 	}
 
 	if (SamplersConfiguration != ERHIBindlessConfiguration::Disabled)
@@ -700,16 +754,6 @@ void FD3D12BindlessDescriptorManager::CleanupResources()
 	}
 }
 
-FRHIDescriptorHandle FD3D12BindlessDescriptorManager::AllocateResourceHandle()
-{
-	if (ResourceManager)
-	{
-		return ResourceManager->Allocate();
-	}
-
-	return FRHIDescriptorHandle();
-}
-
 void FD3D12BindlessDescriptorManager::GarbageCollect()
 {
 	if (ResourceManager)
@@ -726,27 +770,17 @@ void FD3D12BindlessDescriptorManager::Recycle(FD3D12DescriptorHeap* DescriptorHe
 	}
 }
 
-FRHIDescriptorHandle FD3D12BindlessDescriptorManager::AllocateAndInitialize(FD3D12SamplerState* SamplerState)
-{
-	if (SamplerManager)
-	{
-		return SamplerManager->AllocateAndInitialize(SamplerState);
-	}
-
-	return FRHIDescriptorHandle();
-}
-
 void FD3D12BindlessDescriptorManager::ImmediateFree(FRHIDescriptorHandle InHandle)
 {
 	if (InHandle.GetType() == ERHIDescriptorHeapType::Standard && ResourceManager)
 	{
-		ResourceManager->Free(InHandle);
+		GetParentDevice()->GetParentAdapter()->GetBindlessManager().FreeResourceHandle(InHandle);
 		return;
 	}
 
 	if (InHandle.GetType() == ERHIDescriptorHeapType::Sampler && SamplerManager)
 	{
-		SamplerManager->Free(InHandle);
+		GetParentDevice()->GetParentAdapter()->GetBindlessManager().FreeSamplerHandle(InHandle);
 		return;
 	}
 
diff --git a/Engine/Source/Runtime/D3D12RHI/Private/D3D12BindlessDescriptors.h b/Engine/Source/Runtime/D3D12RHI/Private/D3D12BindlessDescriptors.h
--- a/Engine/Source/Runtime/D3D12RHI/Private/D3D12BindlessDescriptors.h
+++ b/Engine/Source/Runtime/D3D12RHI/Private/D3D12BindlessDescriptors.h
@@ -26,6 +26,46 @@ namespace UE::D3D12BindlessDescriptors
 	void DeferredFreeHeap(FD3D12Device* InDevice, FD3D12DescriptorHeap* InHeap);
 }
 
+/** Manager for configuration settings and shared descriptor allocators, stored on the adapter. */
+class FD3D12BindlessManagerAdapter : public FD3D12AdapterChild
+{
+public:
+	void Init(FD3D12Adapter* InParent);
+
+	ERHIBindlessConfiguration GetResourcesConfiguration() const { return BindlessResourcesConfiguration; }
+	ERHIBindlessConfiguration GetSamplersConfiguration() const { return BindlessSamplersConfiguration; }
+
+	bool GetResourcesAllowed() const { return bBindlessResourcesAllowed; }
+	bool GetSamplersAllowed() const { return bBindlessSamplersAllowed; }
+
+	// Bindless descriptor allocators are stored in the adapter, so descriptor handles can be allocated once and shared for multi-GPU objects
+	FRHIDescriptorHandle AllocateSamplerHandle();
+	FRHIDescriptorHandle AllocateResourceHandle();
+	void FreeSamplerHandle(FRHIDescriptorHandle Handle);
+	void FreeResourceHandle(FRHIDescriptorHandle Handle);
+
+	FCriticalSection& GetResourceHeapsCS() { return ResourceHeapsCS; }
+
+	uint32 GetResourceCapacity() const { return ResourceAllocator->GetCapacity(); }
+
+	bool GetResourceAllocatedRange(FRHIDescriptorAllocatorRange& OutAllocatedRange) const { return ResourceAllocator->GetAllocatedRange(OutAllocatedRange); }
+
+	FRHIDescriptorHandle& GetSamplerDefaultHandle() { return SamplerDefaultHandle; }
+
+private:
+	bool bBindlessResourcesAllowed = false;
+	bool bBindlessSamplersAllowed = false;
+
+	ERHIBindlessConfiguration BindlessResourcesConfiguration{};
+	ERHIBindlessConfiguration BindlessSamplersConfiguration{};
+
+	FCriticalSection ResourceHeapsCS;
+	FRHIHeapDescriptorAllocator* ResourceAllocator = nullptr;
+	FRHIHeapDescriptorAllocator* SamplerAllocator = nullptr;
+
+	FRHIDescriptorHandle SamplerDefaultHandle;
+};
+
 /** Manager specifically for bindless sampler descriptors. */
 class FD3D12BindlessSamplerManager : public FD3D12DeviceChild
 {
@@ -35,9 +75,6 @@ public:
 
 	void CleanupResources();
 
-	FRHIDescriptorHandle AllocateAndInitialize(FD3D12SamplerState* SamplerState);
-	void                 Free(FRHIDescriptorHandle InHandle);
-
 	void OpenCommandList(FD3D12CommandContext& Context);
 	void CloseCommandList(FD3D12CommandContext& Context);
 
@@ -46,6 +83,8 @@ public:
 	FD3D12DescriptorHeap* GetHeap() const { return GpuHeap.GetReference(); }
 	ERHIBindlessConfiguration GetConfiguration() const { return Configuration; }
 
+	void InitializeSampler(FRHIDescriptorHandle DstHandle, FD3D12SamplerState* SamplerState);
+
 private:
 	FD3D12DescriptorHeapPtr      GpuHeap;
 	FRHIHeapDescriptorAllocator  Allocator;
@@ -121,12 +160,9 @@ class FD3D12BindlessResourceManager : public FD3D12DeviceChild
 {
 public:
 	FD3D12BindlessResourceManager() = delete;
-	FD3D12BindlessResourceManager(FD3D12Device* InDevice, ERHIBindlessConfiguration InConfiguration, uint32 InNumDescriptors, TConstArrayView<TStatId> InStats);
+	FD3D12BindlessResourceManager(FD3D12Device* InDevice, FD3D12Adapter* InAdapter);
 
 	void CleanupResources();
-
-	FRHIDescriptorHandle Allocate();
-	void                 Free(FRHIDescriptorHandle InHandle);
 	
 	void GarbageCollect();
 	void Recycle(FD3D12DescriptorHeap* DescriptorHeap);
@@ -145,6 +181,9 @@ public:
 
 	ERHIBindlessConfiguration GetConfiguration() const { return Configuration; }
 
+	// Called from FD3D12Adapter::AllocateBindlessResourceHandle
+	void GrowCPUHeap(uint32 OriginalNumDescriptors, uint32 NewNumDescriptors);
+
 private:
 	void CopyCpuHeap(FD3D12DescriptorHeap* DestinationHeap);
 	void AssignHeapToState(FD3D12ContextBindlessState& State);
@@ -154,12 +193,10 @@ private:
 	int AddActiveGPUHeap();
 	void ReleaseGPUHeaps();
 	void UpdateInUseGPUHeaps(bool bInUse);
-
-	FRHIDescriptorHandle ResizeGrowAndAllocate();
 	
-	FCriticalSection				HeapsCS;
+	// Critical section shared across devices
+	FCriticalSection&				HeapsCS;
 	FD3D12DescriptorHeapPtr         CpuHeap;
-	FRHIHeapDescriptorAllocator     Allocator;
 	const ERHIBindlessConfiguration Configuration;
 
 	uint64 							GarbageCollectCycle = 0;
@@ -215,14 +252,14 @@ public:
 	bool AreResourcesFullyBindless() const { return GetResourcesConfiguration() == ERHIBindlessConfiguration::AllShaders; }
 	bool AreSamplersFullyBindless () const { return GetSamplersConfiguration()  == ERHIBindlessConfiguration::AllShaders; }
 
-	FRHIDescriptorHandle AllocateResourceHandle();
-	FRHIDescriptorHandle AllocateAndInitialize(FD3D12SamplerState* SamplerState);
 	void                 ImmediateFree(FRHIDescriptorHandle InHandle);
 	void                 DeferredFreeFromDestructor(FRHIDescriptorHandle InHandle);
 
 	void GarbageCollect();
 	void Recycle(FD3D12DescriptorHeap* DescriptorHeap);
 
+	void InitializeSampler(FRHIDescriptorHandle DstHandle, FD3D12SamplerState* SamplerState) { SamplerManager->InitializeSampler(DstHandle, SamplerState); }
+
 	void InitializeDescriptor(FRHIDescriptorHandle DstHandle, FD3D12View* View);
 	void UpdateDescriptor(FD3D12ContextArray const& Contexts, FRHIDescriptorHandle DstHandle, FD3D12View* SourceView);
 
@@ -241,6 +278,10 @@ public:
 #endif
 
 private:
+	// Needs direct access to underlying resource manager for resource allocation functions
+	friend class FD3D12BindlessManagerAdapter;
+	FD3D12BindlessResourceManager& GetResourceManager() const { return *ResourceManager; }
+
 	TUniquePtr<FD3D12BindlessResourceManager> ResourceManager;
 	TUniquePtr<FD3D12BindlessSamplerManager>  SamplerManager;
 
@@ -248,4 +289,20 @@ private:
 	ERHIBindlessConfiguration SamplersConfiguration{};
 };
 
-#endif // PLATFORM_SUPPORTS_BINDLESS_RENDERING
+#else // PLATFORM_SUPPORTS_BINDLESS_RENDERING
+
+// Stub version when bindless is compiled out
+class FD3D12BindlessManagerAdapter : public FD3D12AdapterChild
+{
+public:
+	inline void Init(FD3D12Adapter* InParent) { SetParentAdapter(InParent); }
+
+	inline FRHIDescriptorHandle AllocateSamplerHandle() { return FRHIDescriptorHandle(); }
+	inline FRHIDescriptorHandle AllocateResourceHandle() { return FRHIDescriptorHandle(); }
+	inline FRHIDescriptorHandle& GetSamplerDefaultHandle() { return SamplerDefaultHandle; }
+
+private:
+	FRHIDescriptorHandle SamplerDefaultHandle;
+};
+
+#endif // !PLATFORM_SUPPORTS_BINDLESS_RENDERING
diff --git a/Engine/Source/Runtime/D3D12RHI/Private/D3D12Buffer.cpp b/Engine/Source/Runtime/D3D12RHI/Private/D3D12Buffer.cpp
--- a/Engine/Source/Runtime/D3D12RHI/Private/D3D12Buffer.cpp
+++ b/Engine/Source/Runtime/D3D12RHI/Private/D3D12Buffer.cpp
@@ -242,7 +242,7 @@ FD3D12Buffer* FD3D12Adapter::CreateRHIBuffer(
 		const uint32 FirstGPUIndex = InGPUMask.GetFirstIndex();
 
 		FD3D12Buffer* NewBuffer0 = nullptr;
-		BufferOut = CreateLinkedObject<FD3D12Buffer>(InGPUMask, [&](FD3D12Device* Device)
+		BufferOut = CreateLinkedObject<FD3D12Buffer>(InGPUMask, [&](FD3D12Device* Device, FD3D12Buffer* FirstLinkedObject)
 		{
 			FD3D12Buffer* NewBuffer = new FD3D12Buffer(Device, BufferDesc);
 			NewBuffer->BufferAlignment = Alignment;
@@ -271,7 +271,7 @@ FD3D12Buffer* FD3D12Adapter::CreateRHIBuffer(
 	}
 	else
 	{
-		BufferOut = CreateLinkedObject<FD3D12Buffer>(InGPUMask, [&](FD3D12Device* Device)
+		BufferOut = CreateLinkedObject<FD3D12Buffer>(InGPUMask, [&](FD3D12Device* Device, FD3D12Buffer* FirstLinkedObject)
 		{
 			FD3D12Buffer* NewBuffer = new FD3D12Buffer(Device, BufferDesc);
 			NewBuffer->BufferAlignment = Alignment;
@@ -410,7 +410,7 @@ FBufferRHIRef FD3D12DynamicRHI::CreateBuffer(FRHICommandListBase& RHICmdList, FR
 {
 	if (BufferDesc.IsNull())
 	{
-		return GetAdapter().CreateLinkedObject<FD3D12Buffer>(CreateInfo.GPUMask, [BufferDesc](FD3D12Device* Device)
+		return GetAdapter().CreateLinkedObject<FD3D12Buffer>(CreateInfo.GPUMask, [BufferDesc](FD3D12Device* Device, FD3D12Buffer* FirstLinkedObject)
 		{
 			return new FD3D12Buffer(Device, BufferDesc);
 		});
diff --git a/Engine/Source/Runtime/D3D12RHI/Private/D3D12Device.cpp b/Engine/Source/Runtime/D3D12RHI/Private/D3D12Device.cpp
--- a/Engine/Source/Runtime/D3D12RHI/Private/D3D12Device.cpp
+++ b/Engine/Source/Runtime/D3D12RHI/Private/D3D12Device.cpp
@@ -484,7 +484,7 @@ void FD3D12Device::CreateDefaultViews()
 
 	{
 		const FSamplerStateInitializerRHI SamplerDesc(SF_Trilinear, AM_Clamp, AM_Clamp, AM_Clamp);
-		DefaultViews.DefaultSampler = CreateSampler(SamplerDesc);
+		DefaultViews.DefaultSampler = CreateSampler(SamplerDesc, GetGPUIndex() > 0 ? GetParentAdapter()->GetDevice(0)->DefaultViews.DefaultSampler : nullptr);
 
 		// The default sampler must have ID=0
 		// FD3D12DescriptorCache::SetSamplers relies on this
diff --git a/Engine/Source/Runtime/D3D12RHI/Private/D3D12Device.h b/Engine/Source/Runtime/D3D12RHI/Private/D3D12Device.h
--- a/Engine/Source/Runtime/D3D12RHI/Private/D3D12Device.h
+++ b/Engine/Source/Runtime/D3D12RHI/Private/D3D12Device.h
@@ -258,7 +258,7 @@ public:
 	inline FD3D12ResidencyManager& GetResidencyManager() { return ResidencyManager; }
 
 	// Samplers
-	FD3D12SamplerState* CreateSampler(const FSamplerStateInitializerRHI& Initializer);
+	FD3D12SamplerState* CreateSampler(const FSamplerStateInitializerRHI& Initializer, FD3D12SamplerState* FirstLinkedObject);
 	void CreateSamplerInternal(const D3D12_SAMPLER_DESC& Desc, D3D12_CPU_DESCRIPTOR_HANDLE Descriptor);
 
 	// Command Allocators
diff --git a/Engine/Source/Runtime/D3D12RHI/Private/D3D12Query.cpp b/Engine/Source/Runtime/D3D12RHI/Private/D3D12Query.cpp
--- a/Engine/Source/Runtime/D3D12RHI/Private/D3D12Query.cpp
+++ b/Engine/Source/Runtime/D3D12RHI/Private/D3D12Query.cpp
@@ -153,7 +153,7 @@ uint32 FD3D12QueryHeap::Release()
 FRenderQueryRHIRef FD3D12DynamicRHI::RHICreateRenderQuery(ERenderQueryType QueryType)
 {
 	check(QueryType == RQT_Occlusion || QueryType == RQT_AbsoluteTime);
-	return GetAdapter().CreateLinkedObject<FD3D12RenderQuery>(FRHIGPUMask::All(), [QueryType](FD3D12Device* Device)
+	return GetAdapter().CreateLinkedObject<FD3D12RenderQuery>(FRHIGPUMask::All(), [QueryType](FD3D12Device* Device, FD3D12RenderQuery* FirstLinkedObject)
 	{
 		return new FD3D12RenderQuery(Device, QueryType);
 	});
diff --git a/Engine/Source/Runtime/D3D12RHI/Private/D3D12RHICommon.h b/Engine/Source/Runtime/D3D12RHI/Private/D3D12RHICommon.h
--- a/Engine/Source/Runtime/D3D12RHI/Private/D3D12RHICommon.h
+++ b/Engine/Source/Runtime/D3D12RHI/Private/D3D12RHICommon.h
@@ -175,6 +175,10 @@ public:
 		return GetFirstLinkedObject() == this;
 	}
 
+	// The creation lambda function is passed FD3D12Device, and a pointer to the first linked object, or nullptr if the object being constructed
+	// is itself the first object.  The first object pointer is needed for cases where resources need to be shared for all linked objects,
+	// such as bindless resource descriptor handles, which need to be the same across GPUs for use in platform independent logic which generates
+	// descriptor tables.
 	template <typename ReturnType, typename CreationCoreFunction, typename CreationParameterFunction>
 	static ReturnType* CreateLinkedObjects(FRHIGPUMask GPUMask, const CreationParameterFunction& pfnGetCreationParameter, const CreationCoreFunction& pfnCreationCore)
 	{
@@ -183,7 +187,7 @@ public:
 #if WITH_MGPU
 		for (uint32 GPUIndex : GPUMask)
 		{
-			ReturnType* NewObject = pfnCreationCore(pfnGetCreationParameter(GPUIndex));
+			ReturnType* NewObject = pfnCreationCore(pfnGetCreationParameter(GPUIndex), ObjectOut);
 			CA_ASSUME(NewObject != nullptr);
 			if (ObjectOut == nullptr)
 			{
diff --git a/Engine/Source/Runtime/D3D12RHI/Private/D3D12RenderTarget.cpp b/Engine/Source/Runtime/D3D12RHI/Private/D3D12RenderTarget.cpp
--- a/Engine/Source/Runtime/D3D12RHI/Private/D3D12RenderTarget.cpp
+++ b/Engine/Source/Runtime/D3D12RHI/Private/D3D12RenderTarget.cpp
@@ -711,7 +711,7 @@ void FD3D12DynamicRHI::ReadSurfaceDataMSAARaw(FRHITexture* TextureRHI, FIntRect
 	RTVDesc.ViewDimension = D3D12_RTV_DIMENSION_TEXTURE2D;
 	RTVDesc.Texture2D.MipSlice = 0;
 
-	FD3D12RenderTargetView NonMSAARTV(Device);
+	FD3D12RenderTargetView NonMSAARTV(Device, nullptr);			// Always single GPU object, so FirstLinkedObject is nullptr
 	NonMSAARTV.CreateView(&ResourceLocation, RTVDesc);
 
 	// Create a CPU-accessible staging texture to copy the resolved sample data to.
diff --git a/Engine/Source/Runtime/D3D12RHI/Private/D3D12ResourceCollection.cpp b/Engine/Source/Runtime/D3D12RHI/Private/D3D12ResourceCollection.cpp
--- a/Engine/Source/Runtime/D3D12RHI/Private/D3D12ResourceCollection.cpp
+++ b/Engine/Source/Runtime/D3D12RHI/Private/D3D12ResourceCollection.cpp
@@ -8,7 +8,7 @@
 #include "D3D12CommandContext.h"
 #include "D3D12TextureReference.h"
 
-FD3D12ResourceCollection::FD3D12ResourceCollection(FD3D12Device* InParent, FRHICommandListBase& RHICmdList, FD3D12Buffer* InBuffer, TConstArrayView<FRHIResourceCollectionMember> InMembers)
+FD3D12ResourceCollection::FD3D12ResourceCollection(FD3D12Device* InParent, FRHICommandListBase& RHICmdList, FD3D12Buffer* InBuffer, TConstArrayView<FRHIResourceCollectionMember> InMembers, FD3D12ResourceCollection* FirstLinkedObject)
 	: FRHIResourceCollection(InMembers)
 	, FD3D12DeviceChild(InParent)
 	, Buffer(InBuffer->GetLinkedObject(InParent->GetGPUIndex()))
@@ -67,7 +67,7 @@ FD3D12ResourceCollection::FD3D12ResourceCollection(FD3D12Device* InParent, FRHIC
 	SRVDesc.Buffer.NumElements = UE::RHICore::CalculateResourceCollectionMemorySize(InMembers) / 4;
 	SRVDesc.Buffer.Flags = D3D12_BUFFER_SRV_FLAG_RAW;
 
-	BufferSRV = MakeShared<FD3D12ShaderResourceView>(InParent);
+	BufferSRV = MakeShared<FD3D12ShaderResourceView>(InParent, FirstLinkedObject ? FirstLinkedObject->BufferSRV.Get() : nullptr);
 	BufferSRV->CreateView(InBuffer, SRVDesc, FD3D12ShaderResourceView::EFlags::None);
 }
 
@@ -94,9 +94,9 @@ FRHIResourceCollectionRef FD3D12DynamicRHI::RHICreateResourceCollection(FRHIComm
 	FRHIViewDesc::FBufferSRV::FInitializer ViewDesc = FRHIViewDesc::CreateBufferSRV().SetType(FRHIViewDesc::EBufferType::Raw);
 	FShaderResourceViewRHIRef ShaderResourceView = RHICmdList.CreateShaderResourceView(Buffer, ViewDesc);
 
-	return GetAdapter().CreateLinkedObject<FD3D12ResourceCollection>(FRHIGPUMask::All(), [&RHICmdList, Buffer, InMembers](FD3D12Device* Device)
+	return GetAdapter().CreateLinkedObject<FD3D12ResourceCollection>(FRHIGPUMask::All(), [&RHICmdList, Buffer, InMembers](FD3D12Device* Device, FD3D12ResourceCollection* FirstLinkedObject)
 	{
-		return new FD3D12ResourceCollection(Device, RHICmdList, Buffer, InMembers);
+		return new FD3D12ResourceCollection(Device, RHICmdList, Buffer, InMembers, FirstLinkedObject);
 	});
 }
 
diff --git a/Engine/Source/Runtime/D3D12RHI/Private/D3D12ResourceCollection.h b/Engine/Source/Runtime/D3D12RHI/Private/D3D12ResourceCollection.h
--- a/Engine/Source/Runtime/D3D12RHI/Private/D3D12ResourceCollection.h
+++ b/Engine/Source/Runtime/D3D12RHI/Private/D3D12ResourceCollection.h
@@ -10,14 +10,13 @@
 #if PLATFORM_SUPPORTS_BINDLESS_RENDERING
 
 class FD3D12Buffer;
-class FD3D12ResourceCollection;
 class FD3D12RHITextureReference;
 class FD3D12ShaderResourceView;
 
 class FD3D12ResourceCollection : public FRHIResourceCollection, public FD3D12DeviceChild, public FD3D12LinkedAdapterObject<FD3D12ResourceCollection>
 {
 public:
-	FD3D12ResourceCollection(FD3D12Device* InParent, FRHICommandListBase& RHICmdList, FD3D12Buffer* InBuffer, TConstArrayView<FRHIResourceCollectionMember> InMembers);
+	FD3D12ResourceCollection(FD3D12Device* InParent, FRHICommandListBase& RHICmdList, FD3D12Buffer* InBuffer, TConstArrayView<FRHIResourceCollectionMember> InMembers, FD3D12ResourceCollection* FirstLinkedObject);
 	~FD3D12ResourceCollection();
 
 	virtual FRHIDescriptorHandle GetBindlessHandle() const final;
diff --git a/Engine/Source/Runtime/D3D12RHI/Private/D3D12SRV.cpp b/Engine/Source/Runtime/D3D12RHI/Private/D3D12SRV.cpp
--- a/Engine/Source/Runtime/D3D12RHI/Private/D3D12SRV.cpp
+++ b/Engine/Source/Runtime/D3D12RHI/Private/D3D12SRV.cpp
@@ -9,12 +9,12 @@
 //
 // -----------------------------------------------------------------------------------------------------
 
-FD3D12ShaderResourceView::FD3D12ShaderResourceView(FD3D12Device* InDevice)
-	: TD3D12View(InDevice, ERHIDescriptorHeapType::Standard)
+FD3D12ShaderResourceView::FD3D12ShaderResourceView(FD3D12Device* InDevice, FD3D12ShaderResourceView* FirstLinkedObject)
+	: TD3D12View(InDevice, ERHIDescriptorHeapType::Standard, FirstLinkedObject)
 {}
 
-FD3D12ShaderResourceView::FD3D12ShaderResourceView(FD3D12Device* InDevice, FD3D12RayTracingScene* InRayTracingScene)
-	: TD3D12View(InDevice, ERHIDescriptorHeapType::Standard)
+FD3D12ShaderResourceView::FD3D12ShaderResourceView(FD3D12Device* InDevice, FD3D12ShaderResourceView* FirstLinkedObject, FD3D12RayTracingScene* InRayTracingScene)
+	: TD3D12View(InDevice, ERHIDescriptorHeapType::Standard, FirstLinkedObject)
 	, RayTracingScene(InRayTracingScene)
 {}
 
@@ -300,12 +300,12 @@ void FD3D12ShaderResourceView_RHI::UpdateView(FD3D12ContextArray const& Contexts
 	}
 }
 
-FD3D12ShaderResourceView_RHI::FD3D12ShaderResourceView_RHI(FD3D12Device* InDevice, FRHIViewableResource* InResource, FRHIViewDesc const& InViewDesc)
+FD3D12ShaderResourceView_RHI::FD3D12ShaderResourceView_RHI(FD3D12Device* InDevice, FRHIViewableResource* InResource, FRHIViewDesc const& InViewDesc, FD3D12ShaderResourceView_RHI* FirstLinkedObject)
 	: FRHIShaderResourceView(InResource, InViewDesc)
 #if RHI_RAYTRACING
-	, FD3D12ShaderResourceView(InDevice, InViewDesc.Buffer.SRV.BufferType == FRHIViewDesc::EBufferType::AccelerationStructure ? FD3D12DynamicRHI::ResourceCast(InViewDesc.Buffer.SRV.RayTracingScene) : nullptr)
+	, FD3D12ShaderResourceView(InDevice, FirstLinkedObject, InViewDesc.Buffer.SRV.BufferType == FRHIViewDesc::EBufferType::AccelerationStructure ? FD3D12DynamicRHI::ResourceCast(InViewDesc.Buffer.SRV.RayTracingScene) : nullptr)
 #else
-	, FD3D12ShaderResourceView(InDevice)
+	, FD3D12ShaderResourceView(InDevice, FirstLinkedObject)
 #endif
 {}
 
@@ -323,13 +323,13 @@ FShaderResourceViewRHIRef FD3D12DynamicRHI::RHICreateShaderResourceView(class FR
 		? FD3D12DynamicRHI::ResourceCast(static_cast<FRHIBuffer* >(Resource))->GetLinkedObjectsGPUMask()
 		: FD3D12DynamicRHI::ResourceCast(static_cast<FRHITexture*>(Resource))->GetLinkedObjectsGPUMask();
 
-	FD3D12ShaderResourceView_RHI* View = GetAdapter().CreateLinkedObject<FD3D12ShaderResourceView_RHI>(RelevantGPUs, [&](FD3D12Device* Device)
+	FD3D12ShaderResourceView_RHI* View = GetAdapter().CreateLinkedObject<FD3D12ShaderResourceView_RHI>(RelevantGPUs, [&](FD3D12Device* Device, FD3D12ShaderResourceView_RHI* FirstLinkedObject)
 	{
 		FRHIViewableResource* TargetResource = ViewDesc.IsBuffer()
 			? static_cast<FRHIViewableResource*>(FD3D12DynamicRHI::ResourceCast(static_cast<FRHIBuffer* >(Resource), Device->GetGPUIndex()))
 			: static_cast<FRHIViewableResource*>(FD3D12DynamicRHI::ResourceCast(static_cast<FRHITexture*>(Resource), Device->GetGPUIndex()));
 
-		return new FD3D12ShaderResourceView_RHI(Device, TargetResource, ViewDesc);
+		return new FD3D12ShaderResourceView_RHI(Device, TargetResource, ViewDesc, FirstLinkedObject);
 	});
 
 	View->CreateViews(RHICmdList);
diff --git a/Engine/Source/Runtime/D3D12RHI/Private/D3D12State.cpp b/Engine/Source/Runtime/D3D12RHI/Private/D3D12State.cpp
--- a/Engine/Source/Runtime/D3D12RHI/Private/D3D12State.cpp
+++ b/Engine/Source/Runtime/D3D12RHI/Private/D3D12State.cpp
@@ -227,9 +227,9 @@ FSamplerStateRHIRef FD3D12DynamicRHI::RHICreateSamplerState(const FSamplerStateI
 {
 	FD3D12Adapter* Adapter = &GetAdapter();
 
-	return Adapter->CreateLinkedObject<FD3D12SamplerState>(FRHIGPUMask::All(), [&](FD3D12Device* Device)
+	return Adapter->CreateLinkedObject<FD3D12SamplerState>(FRHIGPUMask::All(), [&](FD3D12Device* Device, FD3D12SamplerState* FirstLinkedObject)
 	{
-		return Device->CreateSampler(Initializer);
+		return Device->CreateSampler(Initializer, FirstLinkedObject);
 	});
 }
 
@@ -257,7 +257,7 @@ static void LogSamplerStateWarning(const FSamplerStateInitializerRHI& Initialize
 	);
 }
 
-FD3D12SamplerState* FD3D12Device::CreateSampler(const FSamplerStateInitializerRHI& Initializer)
+FD3D12SamplerState* FD3D12Device::CreateSampler(const FSamplerStateInitializerRHI& Initializer, FD3D12SamplerState* FirstLinkedObject)
 {
 	D3D12_SAMPLER_DESC SamplerDesc;
 	FMemory::Memzero(&SamplerDesc, sizeof(D3D12_SAMPLER_DESC));
@@ -324,7 +324,7 @@ FD3D12SamplerState* FD3D12Device::CreateSampler(const FSamplerStateInitializerRH
 			LogSamplerStateWarning(Initializer);
 		}
 
-		FD3D12SamplerState* NewSampler = new FD3D12SamplerState(this, SamplerDesc, static_cast<uint16>(SamplerID));
+		FD3D12SamplerState* NewSampler = new FD3D12SamplerState(this, SamplerDesc, static_cast<uint16>(SamplerID), FirstLinkedObject);
 
 		SamplerMap.Add(SamplerDesc, NewSampler);
 
@@ -790,7 +790,7 @@ TRefCountPtr<FRHIComputePipelineState> FD3D12DynamicRHI::RHICreateComputePipelin
 	return PSOCache.CreateAndAdd(ComputeShader, RootSignature, LowLevelDesc);
 }
 
-FD3D12SamplerState::FD3D12SamplerState(FD3D12Device* InParent, const D3D12_SAMPLER_DESC& Desc, uint16 SamplerID)
+FD3D12SamplerState::FD3D12SamplerState(FD3D12Device* InParent, const D3D12_SAMPLER_DESC& Desc, uint16 SamplerID, FD3D12SamplerState* FirstLinkedObject)
 	: FD3D12DeviceChild(InParent)
 	, ID(SamplerID)
 {
@@ -800,7 +800,13 @@ FD3D12SamplerState::FD3D12SamplerState(FD3D12Device* InParent, const D3D12_SAMPL
 	GetParentDevice()->CreateSamplerInternal(Desc, OfflineDescriptor);
 
 #if PLATFORM_SUPPORTS_BINDLESS_RENDERING
-	BindlessHandle = GetParentDevice()->GetBindlessDescriptorManager().AllocateAndInitialize(this);
+	FD3D12Adapter* Adapter = InParent->GetParentAdapter();
+	if (Adapter->GetBindlessManager().GetSamplersConfiguration() != ERHIBindlessConfiguration::Disabled)
+	{
+		BindlessHandle = FirstLinkedObject ? FirstLinkedObject->BindlessHandle : InParent->GetParentAdapter()->GetBindlessManager().AllocateSamplerHandle();
+
+		InParent->GetBindlessDescriptorManager().InitializeSampler(BindlessHandle, this);
+	}
 #endif
 }
 
@@ -812,7 +818,8 @@ FD3D12SamplerState::~FD3D12SamplerState()
 		OfflineAllocator.FreeHeapSlot(OfflineDescriptor);
 
 #if PLATFORM_SUPPORTS_BINDLESS_RENDERING
-		if (BindlessHandle.IsValid())
+		// Handle is shared -- freeing is handled by the head link in the FD3D12LinkedAdapterObject
+		if (BindlessHandle.IsValid() && IsHeadLink())
 		{
 			GetParentDevice()->GetBindlessDescriptorManager().DeferredFreeFromDestructor(BindlessHandle);
 		}
diff --git a/Engine/Source/Runtime/D3D12RHI/Private/D3D12State.h b/Engine/Source/Runtime/D3D12RHI/Private/D3D12State.h
--- a/Engine/Source/Runtime/D3D12RHI/Private/D3D12State.h
+++ b/Engine/Source/Runtime/D3D12RHI/Private/D3D12State.h
@@ -17,7 +17,7 @@ public:
 	const uint16 ID;
 
 	FD3D12SamplerState() = delete;
-	FD3D12SamplerState(FD3D12Device* InParent, const D3D12_SAMPLER_DESC& Desc, uint16 SamplerID);
+	FD3D12SamplerState(FD3D12Device* InParent, const D3D12_SAMPLER_DESC& Desc, uint16 SamplerID, FD3D12SamplerState* FirstLinkedObject);
 	~FD3D12SamplerState();
 
 	virtual FRHIDescriptorHandle GetBindlessHandle() const final { return BindlessHandle; }
diff --git a/Engine/Source/Runtime/D3D12RHI/Private/D3D12Texture.cpp b/Engine/Source/Runtime/D3D12RHI/Private/D3D12Texture.cpp
--- a/Engine/Source/Runtime/D3D12RHI/Private/D3D12Texture.cpp
+++ b/Engine/Source/Runtime/D3D12RHI/Private/D3D12Texture.cpp
@@ -626,7 +626,7 @@ FD3D12Texture* FD3D12DynamicRHI::CreateD3D12Texture(const FRHITextureCreateDesc&
 	const D3D12_RESOURCE_STATES CreateState = (CreateDesc.BulkData != nullptr) ? D3D12_RESOURCE_STATE_COPY_DEST : InitialState;
 
 	FD3D12Adapter* Adapter = &GetAdapter();
-	FD3D12Texture* D3D12TextureOut = Adapter->CreateLinkedObject<FD3D12Texture>(CreateDesc.GPUMask, [&](FD3D12Device* Device)
+	FD3D12Texture* D3D12TextureOut = Adapter->CreateLinkedObject<FD3D12Texture>(CreateDesc.GPUMask, [&](FD3D12Device* Device, FD3D12Texture* FirstLinkedObject)
 	{
 		FD3D12Texture* NewTexture = CreateNewD3D12Texture(CreateDesc, Device);
 
@@ -714,7 +714,7 @@ FD3D12Texture* FD3D12DynamicRHI::CreateD3D12Texture(const FRHITextureCreateDesc&
 			Adapter->CreateUAVAliasResource(ClearValuePtr, CreateDesc.DebugName, Location);
 		}
 
-		NewTexture->CreateViews();
+		NewTexture->CreateViews(FirstLinkedObject);
 
 #if WITH_GPUDEBUGCRASH
 		if (EnumHasAnyFlags(CreateDesc.Flags, TexCreate_Invalid))
@@ -862,7 +862,7 @@ FTextureRHIRef FD3D12DynamicRHI::RHIAsyncCreateTexture2D(uint32 SizeX, uint32 Si
 	const D3D12_RESOURCE_STATES InitialState = D3D12_RESOURCE_STATE_COMMON;
 
 	FD3D12Adapter* Adapter = &GetAdapter();
-	FD3D12Texture* TextureOut = Adapter->CreateLinkedObject<FD3D12Texture>(FRHIGPUMask::All(), [&](FD3D12Device* Device)
+	FD3D12Texture* TextureOut = Adapter->CreateLinkedObject<FD3D12Texture>(FRHIGPUMask::All(), [&](FD3D12Device* Device, FD3D12Texture* FirstLinkedObject)
 	{
 		FD3D12Texture* NewTexture = CreateNewD3D12Texture(CreateDesc, Device);
 
@@ -886,7 +886,7 @@ FTextureRHIRef FD3D12DynamicRHI::RHIAsyncCreateTexture2D(uint32 SizeX, uint32 Si
 		SRVDesc.Texture2D.PlaneSlice = UE::DXGIUtilities::GetPlaneSliceFromViewFormat(PlatformResourceFormat, SRVDesc.Format);
 
 		// Create a wrapper for the SRV and set it on the texture
-		NewTexture->EmplaceSRV(SRVDesc);
+		NewTexture->EmplaceSRV(SRVDesc, FirstLinkedObject);
 
 		return NewTexture;
 	});
@@ -1250,7 +1250,7 @@ FRHIDescriptorHandle FD3D12Texture::GetDefaultBindlessHandle() const
 }
 
 
-void FD3D12Texture::CreateViews()
+void FD3D12Texture::CreateViews(FD3D12Texture* FirstLinkedObject)
 {
 	TRACE_CPUPROFILER_EVENT_SCOPE(D3D12RHI::CreateViews);
 
@@ -1309,7 +1309,7 @@ void FD3D12Texture::CreateViews()
 			RTVDesc.Texture3D.WSize = Desc.Depth;
 
 			SetNumRTVs(1);
-			EmplaceRTV(RTVDesc, 0);
+			EmplaceRTV(RTVDesc, 0, FirstLinkedObject);
 		}
 		else
 		{
@@ -1336,7 +1336,7 @@ void FD3D12Texture::CreateViews()
 						RTVDesc.Texture2DArray.MipSlice = MipIndex;
 						RTVDesc.Texture2DArray.PlaneSlice = UE::DXGIUtilities::GetPlaneSliceFromViewFormat(PlatformResourceFormat, RTVDesc.Format);
 
-						EmplaceRTV(RTVDesc, RTVIndex++);
+						EmplaceRTV(RTVDesc, RTVIndex++, FirstLinkedObject);
 					}
 				}
 				else
@@ -1378,7 +1378,7 @@ void FD3D12Texture::CreateViews()
 						}
 					}
 
-					EmplaceRTV(RTVDesc, RTVIndex++);
+					EmplaceRTV(RTVDesc, RTVIndex++, FirstLinkedObject);
 				}
 			}
 		}
@@ -1431,7 +1431,7 @@ void FD3D12Texture::CreateViews()
 				DSVDesc.Flags |= (AccessType & FExclusiveDepthStencil::DepthWrite_StencilRead) ? D3D12_DSV_FLAG_READ_ONLY_STENCIL : D3D12_DSV_FLAG_NONE;
 			}
 
-			EmplaceDSV(DSVDesc, AccessType);
+			EmplaceDSV(DSVDesc, AccessType, FirstLinkedObject);
 		}
 	}
 
@@ -1496,7 +1496,7 @@ void FD3D12Texture::CreateViews()
 			}
 		}
 
-		EmplaceSRV(SRVDesc);
+		EmplaceSRV(SRVDesc, FirstLinkedObject);
 	}
 }
 
@@ -2461,7 +2461,7 @@ FD3D12Texture* FD3D12DynamicRHI::CreateTextureFromResource(bool bTextureArray, b
 		.SetNumSamples(NumSamples)
 		.SetInitialState(ERHIAccess::SRVMask);
 
-	FD3D12Texture* Texture2D = Adapter->CreateLinkedObject<FD3D12Texture>(Device->GetGPUMask(), [&](FD3D12Device* Device)
+	FD3D12Texture* Texture2D = Adapter->CreateLinkedObject<FD3D12Texture>(Device->GetGPUMask(), [&](FD3D12Device* Device, FD3D12Texture* FirstLinkedObject)
 	{
 		return CreateNewD3D12Texture(CreateDesc, Device);
 	});
@@ -2471,7 +2471,7 @@ FD3D12Texture* FD3D12DynamicRHI::CreateTextureFromResource(bool bTextureArray, b
 	Location.SetResource(TextureResource);
 	Location.SetGPUVirtualAddress(TextureResource->GetGPUVirtualAddress());
 
-	Texture2D->CreateViews();
+	Texture2D->CreateViews(nullptr);		// Always single GPU object, so FirstLinkedObject is nullptr
 
 	FD3D12TextureStats::D3D12TextureAllocated(*Texture2D);
 
@@ -2546,7 +2546,7 @@ FD3D12Texture* FD3D12DynamicRHI::CreateAliasedD3D12Texture2D(FD3D12Texture* Sour
 	const FString Name = SourceTexture->GetName().ToString() + TEXT("Alias");
 	FRHITextureCreateDesc CreateDesc(SourceTexture->GetDesc(), ERHIAccess::SRVMask, *Name);
 
-	FD3D12Texture* Texture2D = Adapter->CreateLinkedObject<FD3D12Texture>(Device->GetGPUMask(), [&](FD3D12Device* Device)
+	FD3D12Texture* Texture2D = Adapter->CreateLinkedObject<FD3D12Texture>(Device->GetGPUMask(), [&](FD3D12Device* Device, FD3D12Texture* FirstLinkedObject)
 	{
 		return CreateNewD3D12Texture(CreateDesc, Device);
 	});
diff --git a/Engine/Source/Runtime/D3D12RHI/Private/D3D12Texture.h b/Engine/Source/Runtime/D3D12RHI/Private/D3D12Texture.h
--- a/Engine/Source/Runtime/D3D12RHI/Private/D3D12Texture.h
+++ b/Engine/Source/Runtime/D3D12RHI/Private/D3D12Texture.h
@@ -76,7 +76,7 @@ public:
 		
 	// Setup functionality
 	void InitializeTextureData(FRHICommandListBase& RHICmdList, const FRHITextureCreateDesc& CreateDesc, D3D12_RESOURCE_STATES DestinationState);
-	void CreateViews();
+	void CreateViews(FD3D12Texture* FirstLinkedObject);
 	void SetCreatedRTVsPerSlice(bool Value, int32 InRTVArraySize)
 	{
 		bCreatedRTVsPerSlice = Value;
@@ -88,25 +88,25 @@ public:
 		RenderTargetViews.SetNum(Num);
 	}
 
-	void EmplaceRTV(D3D12_RENDER_TARGET_VIEW_DESC const& RTVDesc, int32 Index)
+	void EmplaceRTV(D3D12_RENDER_TARGET_VIEW_DESC const& RTVDesc, int32 Index, FD3D12Texture* FirstLinkedObject)
 	{
 		check(RenderTargetViews.IsValidIndex(Index));
 		check(!RenderTargetViews[Index]);
 
-		RenderTargetViews[Index] = MakeShared<FD3D12RenderTargetView>(GetParentDevice());
+		RenderTargetViews[Index] = MakeShared<FD3D12RenderTargetView>(GetParentDevice(), FirstLinkedObject ? FirstLinkedObject->RenderTargetViews[Index].Get() : nullptr);
 		RenderTargetViews[Index]->CreateView(this, RTVDesc);
 	}
 
-	void EmplaceDSV(D3D12_DEPTH_STENCIL_VIEW_DESC const& DSVDesc, int32 Index)
+	void EmplaceDSV(D3D12_DEPTH_STENCIL_VIEW_DESC const& DSVDesc, int32 Index, FD3D12Texture* FirstLinkedObject)
 	{
 		check(Index < FExclusiveDepthStencil::MaxIndex);
 		check(!DepthStencilViews[Index]);
 
-		DepthStencilViews[Index] = MakeShared<FD3D12DepthStencilView>(GetParentDevice());
+		DepthStencilViews[Index] = MakeShared<FD3D12DepthStencilView>(GetParentDevice(), FirstLinkedObject ? FirstLinkedObject->DepthStencilViews[Index].Get() : nullptr);
 		DepthStencilViews[Index]->CreateView(this, DSVDesc);
 	}
 
-	void EmplaceSRV(D3D12_SHADER_RESOURCE_VIEW_DESC const& SRVDesc)
+	void EmplaceSRV(D3D12_SHADER_RESOURCE_VIEW_DESC const& SRVDesc, FD3D12Texture* FirstLinkedObject)
 	{
 		check(!ShaderResourceView);
 
@@ -114,7 +114,7 @@ public:
 			? FD3D12ShaderResourceView::EFlags::SkipFastClearFinalize
 			: FD3D12ShaderResourceView::EFlags::None;
 
-		ShaderResourceView = MakeShared<FD3D12ShaderResourceView>(GetParentDevice());
+		ShaderResourceView = MakeShared<FD3D12ShaderResourceView>(GetParentDevice(), FirstLinkedObject ? FirstLinkedObject->ShaderResourceView.Get() : nullptr);
 		ShaderResourceView->CreateView(this, SRVDesc, Flags);
 	}
 
diff --git a/Engine/Source/Runtime/D3D12RHI/Private/D3D12TextureReference.cpp b/Engine/Source/Runtime/D3D12RHI/Private/D3D12TextureReference.cpp
--- a/Engine/Source/Runtime/D3D12RHI/Private/D3D12TextureReference.cpp
+++ b/Engine/Source/Runtime/D3D12RHI/Private/D3D12TextureReference.cpp
@@ -2,10 +2,10 @@
 
 #include "D3D12TextureReference.h"
 
-FD3D12RHITextureReference::FD3D12RHITextureReference(FD3D12Device* InDevice, FD3D12Texture* InReferencedTexture)
+FD3D12RHITextureReference::FD3D12RHITextureReference(FD3D12Device* InDevice, FD3D12Texture* InReferencedTexture, FD3D12RHITextureReference* FirstLinkedObject)
 	: FD3D12DeviceChild(InDevice)
 #if PLATFORM_SUPPORTS_BINDLESS_RENDERING
-	, FRHITextureReference(InReferencedTexture, InDevice->GetBindlessDescriptorManager().AllocateResourceHandle())
+	, FRHITextureReference(InReferencedTexture, FirstLinkedObject ? FirstLinkedObject->BindlessHandle : InDevice->GetParentAdapter()->GetBindlessManager().AllocateResourceHandle())
 #else
 	, FRHITextureReference(InReferencedTexture)
 #endif
@@ -30,7 +30,11 @@ FD3D12RHITextureReference::~FD3D12RHITextureReference()
 	{
 		FD3D12DynamicRHI::ResourceCast(GetReferencedTexture())->RemoveRenameListener(this);
 
-		GetParentDevice()->GetBindlessDescriptorManager().DeferredFreeFromDestructor(BindlessHandle);
+		// Bindless handle is shared, head link object handles freeing
+		if (IsHeadLink())
+		{
+			GetParentDevice()->GetBindlessDescriptorManager().DeferredFreeFromDestructor(BindlessHandle);
+		}
 	}
 #endif // PLATFORM_SUPPORTS_BINDLESS_RENDERING
 }
@@ -90,9 +94,9 @@ FTextureReferenceRHIRef FD3D12DynamicRHI::RHICreateTextureReference(FRHICommandL
 	FRHITexture* ReferencedTexture = InReferencedTexture ? InReferencedTexture : FRHITextureReference::GetDefaultTexture();
 
 	FD3D12Adapter* Adapter = &GetAdapter();
-	return Adapter->CreateLinkedObject<FD3D12RHITextureReference>(FRHIGPUMask::All(), [ReferencedTexture](FD3D12Device* Device)
+	return Adapter->CreateLinkedObject<FD3D12RHITextureReference>(FRHIGPUMask::All(), [ReferencedTexture](FD3D12Device* Device, FD3D12RHITextureReference* FirstLinkedObject)
 	{
-		return new FD3D12RHITextureReference(Device, ResourceCast(ReferencedTexture, Device->GetGPUIndex()));
+		return new FD3D12RHITextureReference(Device, ResourceCast(ReferencedTexture, Device->GetGPUIndex()), FirstLinkedObject);
 	});
 }
 
diff --git a/Engine/Source/Runtime/D3D12RHI/Private/D3D12TextureReference.h b/Engine/Source/Runtime/D3D12RHI/Private/D3D12TextureReference.h
--- a/Engine/Source/Runtime/D3D12RHI/Private/D3D12TextureReference.h
+++ b/Engine/Source/Runtime/D3D12RHI/Private/D3D12TextureReference.h
@@ -12,7 +12,7 @@ class FD3D12RHITextureReference : public FD3D12DeviceChild, public FRHITextureRe
 {
 public:
 	FD3D12RHITextureReference() = delete;
-	FD3D12RHITextureReference(FD3D12Device* InDevice, FD3D12Texture* InReferencedTexture);
+	FD3D12RHITextureReference(FD3D12Device* InDevice, FD3D12Texture* InReferencedTexture, FD3D12RHITextureReference* FirstLinkedObject);
 	~FD3D12RHITextureReference();
 
 #if PLATFORM_SUPPORTS_BINDLESS_RENDERING
diff --git a/Engine/Source/Runtime/D3D12RHI/Private/D3D12TransientResourceAllocator.cpp b/Engine/Source/Runtime/D3D12RHI/Private/D3D12TransientResourceAllocator.cpp
--- a/Engine/Source/Runtime/D3D12RHI/Private/D3D12TransientResourceAllocator.cpp
+++ b/Engine/Source/Runtime/D3D12RHI/Private/D3D12TransientResourceAllocator.cpp
@@ -147,7 +147,7 @@ FRHITransientHeap* FD3D12TransientHeapCache::CreateHeap(const FRHITransientHeap:
 	// If heap is flagged for NNE buffers, make it visible on first GPU only.  Required by DirectML.
 	FRHIGPUMask VisibleNodeMask = HeapInitializer.Flags == ERHITransientHeapFlags::AllowNNEBuffers ? FRHIGPUMask::GPU0() : FRHIGPUMask::All();
 
-	return GetParentAdapter()->CreateLinkedObject<FD3D12TransientHeap>(VisibleNodeMask, [&](FD3D12Device* Device)
+	return GetParentAdapter()->CreateLinkedObject<FD3D12TransientHeap>(VisibleNodeMask, [&](FD3D12Device* Device, FD3D12TransientHeap* FirstLinkedObject)
 	{
 		return new FD3D12TransientHeap(HeapInitializer, GetParentAdapter(), Device, VisibleNodeMask);
 	});
diff --git a/Engine/Source/Runtime/D3D12RHI/Private/D3D12UAV.cpp b/Engine/Source/Runtime/D3D12RHI/Private/D3D12UAV.cpp
--- a/Engine/Source/Runtime/D3D12RHI/Private/D3D12UAV.cpp
+++ b/Engine/Source/Runtime/D3D12RHI/Private/D3D12UAV.cpp
@@ -9,8 +9,8 @@
 //
 // -----------------------------------------------------------------------------------------------------
 
-FD3D12UnorderedAccessView::FD3D12UnorderedAccessView(FD3D12Device* InDevice)
-	: TD3D12View(InDevice, ERHIDescriptorHeapType::Standard)
+FD3D12UnorderedAccessView::FD3D12UnorderedAccessView(FD3D12Device* InDevice, FD3D12UnorderedAccessView* FirstLinkedObject)
+	: TD3D12View(InDevice, ERHIDescriptorHeapType::Standard, FirstLinkedObject)
 {}
 
 void FD3D12UnorderedAccessView::UpdateResourceInfo(const FResourceInfo& InResource, const D3D12_UNORDERED_ACCESS_VIEW_DESC& InD3DViewDesc, EFlags InFlags)
@@ -238,9 +238,9 @@ void FD3D12UnorderedAccessView_RHI::UpdateView(FD3D12ContextArray const& Context
 	}
 }
 
-FD3D12UnorderedAccessView_RHI::FD3D12UnorderedAccessView_RHI(FD3D12Device* InDevice, FRHIViewableResource* InResource, FRHIViewDesc const& InViewDesc)
+FD3D12UnorderedAccessView_RHI::FD3D12UnorderedAccessView_RHI(FD3D12Device* InDevice, FRHIViewableResource* InResource, FRHIViewDesc const& InViewDesc, FD3D12UnorderedAccessView_RHI* FirstLinkedObject)
 	: FRHIUnorderedAccessView(InResource, InViewDesc)
-	, FD3D12UnorderedAccessView(InDevice)
+	, FD3D12UnorderedAccessView(InDevice, FirstLinkedObject)
 {}
 
 
@@ -257,13 +257,13 @@ FUnorderedAccessViewRHIRef FD3D12DynamicRHI::RHICreateUnorderedAccessView(FRHICo
 		? FD3D12DynamicRHI::ResourceCast(static_cast<FRHIBuffer* >(Resource))->GetLinkedObjectsGPUMask()
 		: FD3D12DynamicRHI::ResourceCast(static_cast<FRHITexture*>(Resource))->GetLinkedObjectsGPUMask();
 
-	FD3D12UnorderedAccessView_RHI* View = GetAdapter().CreateLinkedObject<FD3D12UnorderedAccessView_RHI>(RelevantGPUs, [&](FD3D12Device* Device)
+	FD3D12UnorderedAccessView_RHI* View = GetAdapter().CreateLinkedObject<FD3D12UnorderedAccessView_RHI>(RelevantGPUs, [&](FD3D12Device* Device, FD3D12UnorderedAccessView_RHI* FirstLinkedObject)
 	{
 		FRHIViewableResource* TargetResource = ViewDesc.IsBuffer()
 			? static_cast<FRHIViewableResource*>(FD3D12DynamicRHI::ResourceCast(static_cast<FRHIBuffer* >(Resource), Device->GetGPUIndex()))
 			: static_cast<FRHIViewableResource*>(FD3D12DynamicRHI::ResourceCast(static_cast<FRHITexture*>(Resource), Device->GetGPUIndex()));
 
-		return new FD3D12UnorderedAccessView_RHI(Device, TargetResource, ViewDesc);
+		return new FD3D12UnorderedAccessView_RHI(Device, TargetResource, ViewDesc, FirstLinkedObject);
 	});
 
 	View->CreateViews(RHICmdList);
@@ -356,7 +356,7 @@ void FD3D12CommandContext::ClearUAV(TRHICommandList_RecursiveHazardous<FD3D12Com
 				}
 
 				// Scoped view will free the offline CPU handle once we return
-				FD3D12UnorderedAccessView UAV(ParentDevice);
+				FD3D12UnorderedAccessView UAV(ParentDevice, nullptr);			// Always single GPU object, so FirstLinkedObject is nullptr
 				UAV.CreateView(UnorderedAccessView->GetResourceLocation(), R32UAVDesc, FD3D12UnorderedAccessView::EFlags::None);
 
 				FD3D12OfflineDescriptor OfflineHandle = UAV.GetOfflineCpuHandle();
diff --git a/Engine/Source/Runtime/D3D12RHI/Private/D3D12UniformBuffer.cpp b/Engine/Source/Runtime/D3D12RHI/Private/D3D12UniformBuffer.cpp
--- a/Engine/Source/Runtime/D3D12RHI/Private/D3D12UniformBuffer.cpp
+++ b/Engine/Source/Runtime/D3D12RHI/Private/D3D12UniformBuffer.cpp
@@ -33,7 +33,7 @@ FUniformBufferRHIRef FD3D12DynamicRHI::RHICreateUniformBuffer(const void* Conten
 	//Note: This is not overly efficient in the mGPU case (we create two+ upload locations) but the CPU savings of having no extra indirection to the resource are worth
 	//      it in single node.
 	// Create the uniform buffer
-	FD3D12UniformBuffer* UniformBufferOut = GetAdapter().CreateLinkedObject<FD3D12UniformBuffer>(FRHIGPUMask::All(), [&](FD3D12Device* Device) -> FD3D12UniformBuffer*
+	FD3D12UniformBuffer* UniformBufferOut = GetAdapter().CreateLinkedObject<FD3D12UniformBuffer>(FRHIGPUMask::All(), [&](FD3D12Device* Device, FD3D12UniformBuffer* FirstLinkedObject) -> FD3D12UniformBuffer*
 	{
 		// If NumBytesActualData == 0, this uniform buffer contains no constants, only a resource table.
 		FD3D12UniformBuffer* NewUniformBuffer = new FD3D12UniformBuffer(Device, Layout, Usage);
diff --git a/Engine/Source/Runtime/D3D12RHI/Private/D3D12Util.cpp b/Engine/Source/Runtime/D3D12RHI/Private/D3D12Util.cpp
--- a/Engine/Source/Runtime/D3D12RHI/Private/D3D12Util.cpp
+++ b/Engine/Source/Runtime/D3D12RHI/Private/D3D12Util.cpp
@@ -1363,8 +1363,8 @@ const FD3D12RootSignature* FD3D12Adapter::GetGlobalRayTracingRootSignature(const
 	QBSS.bUseDiagnosticBuffer = true;
 
 #if PLATFORM_SUPPORTS_BINDLESS_RENDERING
-	QBSS.bUseDirectlyIndexedResourceHeap = bBindlessResourcesAllowed;
-	QBSS.bUseDirectlyIndexedSamplerHeap = bBindlessSamplersAllowed;
+	QBSS.bUseDirectlyIndexedResourceHeap = BindlessManager.GetResourcesAllowed();
+	QBSS.bUseDirectlyIndexedSamplerHeap = BindlessManager.GetSamplersAllowed();
 #endif
 
 	QBSSRegisterCounts.SamplerCount = MAX_SAMPLERS;
diff --git a/Engine/Source/Runtime/D3D12RHI/Private/D3D12View.cpp b/Engine/Source/Runtime/D3D12RHI/Private/D3D12View.cpp
--- a/Engine/Source/Runtime/D3D12RHI/Private/D3D12View.cpp
+++ b/Engine/Source/Runtime/D3D12RHI/Private/D3D12View.cpp
@@ -213,15 +213,19 @@ FD3D12ViewRange::FD3D12ViewRange(D3D12_UNORDERED_ACCESS_VIEW_DESC const& ViewDes
 //
 // -----------------------------------------------------------------------------------------------------
 
-FD3D12View::FD3D12View(FD3D12Device* InDevice, ERHIDescriptorHeapType InHeapType)
+FD3D12View::FD3D12View(FD3D12Device* InDevice, ERHIDescriptorHeapType InHeapType, FD3D12View* FirstLinkedObject)
 	: FD3D12DeviceChild(InDevice)
 	, OfflineCpuHandle(InDevice->GetOfflineDescriptorManager(InHeapType).AllocateHeapSlot())
+#if PLATFORM_SUPPORTS_BINDLESS_RENDERING
+	, bFirstLinkedObject(FirstLinkedObject == nullptr)		// If FirstLinkedObject is nullptr, this newly constructed object itself must be the first
+#endif
 	, HeapType(InHeapType)
 {
 #if PLATFORM_SUPPORTS_BINDLESS_RENDERING
 	if (HeapType == ERHIDescriptorHeapType::Standard)
 	{
-		BindlessHandle = InDevice->GetBindlessDescriptorManager().AllocateResourceHandle();
+		// Bindless handle is shared across linked objects, so pull the handle from the first object if available, or allocate if we are the first
+		BindlessHandle = FirstLinkedObject ? FirstLinkedObject->BindlessHandle : InDevice->GetParentAdapter()->GetBindlessManager().AllocateResourceHandle();
 	}
 #endif
 }
@@ -237,7 +241,11 @@ FD3D12View::~FD3D12View()
 #if PLATFORM_SUPPORTS_BINDLESS_RENDERING
 	if (BindlessHandle.IsValid())
 	{
-		GetParentDevice()->GetBindlessDescriptorManager().DeferredFreeFromDestructor(BindlessHandle);
+		// Bindless handle is shared, first linked object handles freeing
+		if (bFirstLinkedObject)
+		{
+			GetParentDevice()->GetBindlessDescriptorManager().DeferredFreeFromDestructor(BindlessHandle);
+		}
 		BindlessHandle = {};
 	}
 #endif
@@ -335,8 +343,8 @@ void FD3D12View::ResourceRenamed(FD3D12ContextArray const& Contexts, FD3D12BaseS
 //
 // -----------------------------------------------------------------------------------------------------
 
-FD3D12ConstantBufferView::FD3D12ConstantBufferView(FD3D12Device* InDevice)
-	: TD3D12View(InDevice, ERHIDescriptorHeapType::Standard)
+FD3D12ConstantBufferView::FD3D12ConstantBufferView(FD3D12Device* InDevice, FD3D12ConstantBufferView* FirstLinkedObject)
+	: TD3D12View(InDevice, ERHIDescriptorHeapType::Standard, FirstLinkedObject)
 {}
 
 void FD3D12ConstantBufferView::CreateView(FResourceInfo const& InResource, uint32 InOffset, uint32 InAlignedSize)
@@ -367,12 +375,12 @@ void FD3D12ConstantBufferView::UpdateDescriptor()
 
 // -----------------------------------------------------------------------------------------------------
 //
-//                                       FD3D12ConstantBufferView                                       
+//                                       FD3D12RenderTargetView                                       
 //
 // -----------------------------------------------------------------------------------------------------
 
-FD3D12RenderTargetView::FD3D12RenderTargetView(FD3D12Device* InDevice)
-	: TD3D12View(InDevice, ERHIDescriptorHeapType::RenderTarget)
+FD3D12RenderTargetView::FD3D12RenderTargetView(FD3D12Device* InDevice, FD3D12RenderTargetView* FirstLinkedObject)
+	: TD3D12View(InDevice, ERHIDescriptorHeapType::RenderTarget, FirstLinkedObject)
 {}
 
 void FD3D12RenderTargetView::UpdateDescriptor()
@@ -394,8 +402,8 @@ void FD3D12RenderTargetView::UpdateDescriptor()
 //
 // -----------------------------------------------------------------------------------------------------
 
-FD3D12DepthStencilView::FD3D12DepthStencilView(FD3D12Device* InDevice)
-	: TD3D12View(InDevice, ERHIDescriptorHeapType::DepthStencil)
+FD3D12DepthStencilView::FD3D12DepthStencilView(FD3D12Device* InDevice, FD3D12DepthStencilView* FirstLinkedObject)
+	: TD3D12View(InDevice, ERHIDescriptorHeapType::DepthStencil, FirstLinkedObject)
 {}
 
 void FD3D12DepthStencilView::UpdateDescriptor()
diff --git a/Engine/Source/Runtime/D3D12RHI/Private/D3D12View.h b/Engine/Source/Runtime/D3D12RHI/Private/D3D12View.h
--- a/Engine/Source/Runtime/D3D12RHI/Private/D3D12View.h
+++ b/Engine/Source/Runtime/D3D12RHI/Private/D3D12View.h
@@ -219,7 +219,7 @@ public:
 #endif
 
 protected:
-	FD3D12View(FD3D12Device* InDevice, ERHIDescriptorHeapType InHeapType);
+	FD3D12View(FD3D12Device* InDevice, ERHIDescriptorHeapType InHeapType, FD3D12View* FirstLinkedObject);
 	virtual ~FD3D12View();
 
 	virtual void ResourceRenamed(FD3D12ContextArray const& Contexts, FD3D12BaseShaderResource* InRenamedResource, FD3D12ResourceLocation* InNewResourceLocation) override;
@@ -240,7 +240,10 @@ protected:
 	FD3D12OfflineDescriptor OfflineCpuHandle;
 
 #if PLATFORM_SUPPORTS_BINDLESS_RENDERING
+	// BindlessHandle is shared across multi-GPU linked objects.  First linked object is responsible for destroying the handle.
+	// The D3D12View class doesn't have access to FD3D12LinkedAdapterObject::IsHeadLink(), so this separate bool is used.
 	FRHIDescriptorHandle BindlessHandle;
+	bool bFirstLinkedObject;
 #endif
 	ERHIDescriptorHeapType const HeapType;
 };
@@ -251,8 +254,8 @@ class TD3D12View : public FD3D12View
 protected:
 	TDesc D3DViewDesc;
 
-	TD3D12View(FD3D12Device* InDevice, ERHIDescriptorHeapType InHeapType)
-		: FD3D12View(InDevice, InHeapType)
+	TD3D12View(FD3D12Device* InDevice, ERHIDescriptorHeapType InHeapType, TD3D12View* FirstLinkedObject)
+		: FD3D12View(InDevice, InHeapType, FirstLinkedObject)
 	{}
 
 	void CreateView(FResourceInfo const& InResource, TDesc const& InD3DViewDesc)
@@ -278,7 +281,7 @@ class FD3D12ConstantBufferView final : public TD3D12View<FD3D12ConstantBufferVie
 public:
 	static constexpr FD3D12OfflineDescriptor FD3D12DefaultViews::*Null { &FD3D12DefaultViews::NullCBV };
 
-	FD3D12ConstantBufferView(FD3D12Device* InParent);
+	FD3D12ConstantBufferView(FD3D12Device* InParent, FD3D12ConstantBufferView* FirstLinkedObject);
 	void CreateView(FResourceInfo const& InResource, uint32 InOffset, uint32 InAlignedSize);
 
 private:
@@ -300,8 +303,8 @@ public:
 	};
 	FRIEND_ENUM_CLASS_FLAGS(EFlags)
 
-	FD3D12ShaderResourceView(FD3D12Device* InDevice);
-	FD3D12ShaderResourceView(FD3D12Device* InDevice, FD3D12RayTracingScene* InRayTracingScene);
+	FD3D12ShaderResourceView(FD3D12Device* InDevice, FD3D12ShaderResourceView* FirstLinkedObject);
+	FD3D12ShaderResourceView(FD3D12Device* InDevice, FD3D12ShaderResourceView* FirstLinkedObject, FD3D12RayTracingScene* InRayTracingScene);
 	void CreateView(FResourceInfo const& InResource, D3D12_SHADER_RESOURCE_VIEW_DESC const& InD3DViewDesc, EFlags InFlags);
 	void UpdateView(FD3D12ContextArray const& Contexts, const FResourceInfo& InResource, const D3D12_SHADER_RESOURCE_VIEW_DESC& InD3DViewDesc, EFlags InFlags);
 
@@ -338,7 +341,7 @@ public:
 	};
 	FRIEND_ENUM_CLASS_FLAGS(EFlags)
 
-	FD3D12UnorderedAccessView(FD3D12Device* InDevice);
+	FD3D12UnorderedAccessView(FD3D12Device* InDevice, FD3D12UnorderedAccessView* FirstLinkedObject);
 	void CreateView(FResourceInfo const& InResource, D3D12_UNORDERED_ACCESS_VIEW_DESC const& InD3DViewDesc, EFlags InFlags);
 	void UpdateView(FD3D12ContextArray const& Contexts, const FResourceInfo& InResource, const D3D12_UNORDERED_ACCESS_VIEW_DESC& InD3DViewDesc, EFlags InFlags);
 
@@ -366,7 +369,7 @@ class FD3D12RenderTargetView final : public TD3D12View<FD3D12RenderTargetView, D
 public:
 	static constexpr FD3D12OfflineDescriptor FD3D12DefaultViews::*Null { &FD3D12DefaultViews::NullRTV };
 
-	FD3D12RenderTargetView(FD3D12Device* InDevice);
+	FD3D12RenderTargetView(FD3D12Device* InDevice, FD3D12RenderTargetView* FirstLinkedObject);
 	using TD3D12View::CreateView;
 
 private:
@@ -378,7 +381,7 @@ class FD3D12DepthStencilView final : public TD3D12View<FD3D12DepthStencilView, D
 public:
 	static constexpr FD3D12OfflineDescriptor FD3D12DefaultViews::*Null { &FD3D12DefaultViews::NullDSV };
 
-	FD3D12DepthStencilView(FD3D12Device* InDevice);
+	FD3D12DepthStencilView(FD3D12Device* InDevice, FD3D12DepthStencilView* FirstLinkedObject);
 	using TD3D12View::CreateView;
 
 	bool HasDepth  () const { return GetViewSubset().HasPlane(0); }
@@ -426,7 +429,7 @@ class FD3D12ShaderResourceView_RHI
 	, public FD3D12DeferredInitView<FD3D12ShaderResourceView_RHI>
 {
 public:
-	FD3D12ShaderResourceView_RHI(FD3D12Device* InDevice, FRHIViewableResource* InResource, FRHIViewDesc const& InViewDesc);
+	FD3D12ShaderResourceView_RHI(FD3D12Device* InDevice, FRHIViewableResource* InResource, FRHIViewDesc const& InViewDesc, FD3D12ShaderResourceView_RHI* FirstLinkedObject);
 
 	virtual void CreateView();
 	virtual void UpdateView(FD3D12ContextArray const& Contexts);
@@ -447,7 +450,7 @@ class FD3D12UnorderedAccessView_RHI
 	, public FD3D12DeferredInitView<FD3D12UnorderedAccessView_RHI>
 {
 public:
-	FD3D12UnorderedAccessView_RHI(FD3D12Device* InDevice, FRHIViewableResource* InResource, FRHIViewDesc const& InViewDesc);
+	FD3D12UnorderedAccessView_RHI(FD3D12Device* InDevice, FRHIViewableResource* InResource, FRHIViewDesc const& InViewDesc, FD3D12UnorderedAccessView_RHI* FirstLinkedObject);
 
 	virtual void CreateView();
 	virtual void UpdateView(FD3D12ContextArray const& Contexts);
diff --git a/Engine/Source/Runtime/D3D12RHI/Private/D3D12Viewport.cpp b/Engine/Source/Runtime/D3D12RHI/Private/D3D12Viewport.cpp
--- a/Engine/Source/Runtime/D3D12RHI/Private/D3D12Viewport.cpp
+++ b/Engine/Source/Runtime/D3D12RHI/Private/D3D12Viewport.cpp
@@ -111,7 +111,7 @@ FD3D12Texture* GetSwapChainSurface(FD3D12Device* Parent, EPixelFormat PixelForma
 
 	FD3D12DynamicRHI* DynamicRHI = FD3D12DynamicRHI::GetD3DRHI();
 
-	FD3D12Texture* SwapChainTexture = Adapter->CreateLinkedObject<FD3D12Texture>(FRHIGPUMask::All(), [&](FD3D12Device* Device)
+	FD3D12Texture* SwapChainTexture = Adapter->CreateLinkedObject<FD3D12Texture>(FRHIGPUMask::All(), [&](FD3D12Device* Device, FD3D12Texture* FirstLinkedObject)
 	{
 		FD3D12Texture* NewTexture = DynamicRHI->CreateNewD3D12Texture(CreateDesc, Device);
 
@@ -158,8 +158,8 @@ FD3D12Texture* GetSwapChainSurface(FD3D12Device* Parent, EPixelFormat PixelForma
 			RTVDescRight.Texture2DArray.ArraySize = 1;
 
 			NewTexture->SetNumRTVs(2);
-			NewTexture->EmplaceRTV(RTVDescLeft, 0);
-			NewTexture->EmplaceRTV(RTVDescRight, 1);
+			NewTexture->EmplaceRTV(RTVDescLeft, 0, FirstLinkedObject);
+			NewTexture->EmplaceRTV(RTVDescRight, 1, FirstLinkedObject);
 		}
 		else
 		{
@@ -170,7 +170,7 @@ FD3D12Texture* GetSwapChainSurface(FD3D12Device* Parent, EPixelFormat PixelForma
 			RTVDesc.Texture2D.MipSlice = 0;
 
 			NewTexture->SetNumRTVs(1);
-			NewTexture->EmplaceRTV(RTVDesc, 0);
+			NewTexture->EmplaceRTV(RTVDesc, 0, FirstLinkedObject);
 		}
 
 		// create a shader resource view to allow using the backbuffer as a texture
@@ -181,7 +181,7 @@ FD3D12Texture* GetSwapChainSurface(FD3D12Device* Parent, EPixelFormat PixelForma
 		SRVDesc.Texture2D.MostDetailedMip = 0;
 		SRVDesc.Texture2D.MipLevels = 1;
 
-		NewTexture->EmplaceSRV(SRVDesc);
+		NewTexture->EmplaceSRV(SRVDesc, FirstLinkedObject);
 
 		return NewTexture;
 	});
diff --git a/Engine/Source/Runtime/D3D12RHI/Private/D3D12WorkGraph.cpp b/Engine/Source/Runtime/D3D12RHI/Private/D3D12WorkGraph.cpp
--- a/Engine/Source/Runtime/D3D12RHI/Private/D3D12WorkGraph.cpp
+++ b/Engine/Source/Runtime/D3D12RHI/Private/D3D12WorkGraph.cpp
@@ -736,7 +736,7 @@ void FD3D12CommandContext::DispatchWorkGraphShaderBundle(FRHIShaderBundle* Shade
 		SRVDesc.Buffer.StructureByteStride = 0;
 		SRVDesc.Buffer.Flags = D3D12_BUFFER_SRV_FLAG_RAW;
 
-		RecordArgBufferSRV = MakeShared<FD3D12ShaderResourceView>(GetParentDevice());
+		RecordArgBufferSRV = MakeShared<FD3D12ShaderResourceView>(GetParentDevice(), nullptr);			// Always single GPU object, so FirstLinkedObject is nullptr
 		RecordArgBufferSRV->CreateView(RecordArgBufferPtr, SRVDesc, FD3D12ShaderResourceView::EFlags::None);
 
 		LocalSRVs[WorkGraphGlobalShader->RecordArgBufferParam.GetBaseIndex()] = RecordArgBufferSRV->GetOfflineCpuHandle();
diff --git a/Engine/Source/Runtime/D3D12RHI/Private/Windows/WindowsD3D12Viewport.cpp b/Engine/Source/Runtime/D3D12RHI/Private/Windows/WindowsD3D12Viewport.cpp
--- a/Engine/Source/Runtime/D3D12RHI/Private/Windows/WindowsD3D12Viewport.cpp
+++ b/Engine/Source/Runtime/D3D12RHI/Private/Windows/WindowsD3D12Viewport.cpp
@@ -244,7 +244,7 @@ FD3D12Texture* FD3D12Viewport::CreateDummyBackBufferTextures(FD3D12Adapter* InAd
 		.SetFlags(ETextureCreateFlags::RenderTargetable | ETextureCreateFlags::Presentable | ETextureCreateFlags::ResolveTargetable)
 		.SetInitialState(ERHIAccess::Present);
 
-	FD3D12Texture* Result = InAdapter->CreateLinkedObject<FD3D12Texture>(FRHIGPUMask::All(), [this, &CreateDesc](FD3D12Device* Device)
+	FD3D12Texture* Result = InAdapter->CreateLinkedObject<FD3D12Texture>(FRHIGPUMask::All(), [this, &CreateDesc](FD3D12Device* Device, FD3D12Texture* FirstLinkedObject)
 		{
 			return new FD3D12BackBufferReferenceTexture2D(CreateDesc, this, Device);
 		});
