Add Insights CPU marker around the D3D12 Present call
 - This function often blocks for a long time if we're GPU bound, so its helpful to see where that is on the RHI thread timeline in Insights.

diff --git a/Engine/Source/Runtime/D3D12RHI/Private/Windows/WindowsD3D12Viewport.cpp b/Engine/Source/Runtime/D3D12RHI/Private/Windows/WindowsD3D12Viewport.cpp
--- a/Engine/Source/Runtime/D3D12RHI/Private/Windows/WindowsD3D12Viewport.cpp
+++ b/Engine/Source/Runtime/D3D12RHI/Private/Windows/WindowsD3D12Viewport.cpp
@@ -362,7 +362,11 @@ HRESULT FD3D12Viewport::PresentInternal(int32 SyncInterval)
 		}
 #endif
 
-		HRESULT PresentHR = SwapChain1->Present(SyncInterval, Flags);
+		HRESULT PresentHR;
+		{
+			TRACE_CPUPROFILER_EVENT_SCOPE(D3D12_Present);
+			PresentHR = SwapChain1->Present(SyncInterval, Flags);
+		}
 
 		UINT PresentID;
 		if (SUCCEEDED(SwapChain1->GetLastPresentCount(&PresentID)))
