* Fix r.GPUCrashDebugging.Breadcrumbs and -gpubreadcrumbs not working (RHI breadcrumbs were always enabled).
* Fix help text for r.GPUCrashDebugging.Breadcrumbs.
* Check for -lightdred on the command line for enabling lightweight DRED, instead of -dred (which was already processed earlier so it always enabled full DRED).
* Defer Aftermath debug info callbacks until a GPU crash actually occurs, so we don't always write lots of nvdbg files to disk when -nvAftermathDumpShaderDebugInfo is used.

#rnx

diff --git a/Engine/Source/Runtime/D3D12RHI/Private/D3D12Adapter.cpp b/Engine/Source/Runtime/D3D12RHI/Private/D3D12Adapter.cpp
--- a/Engine/Source/Runtime/D3D12RHI/Private/D3D12Adapter.cpp
+++ b/Engine/Source/Runtime/D3D12RHI/Private/D3D12Adapter.cpp
@@ -107,7 +107,7 @@ ED3D12DredMode D3D12_GetDredMode()
 	{
 		return ED3D12DredMode::Full;
 	}
-	else if (UE::RHI::ShouldEnableGPUCrashFeature(*CVarD3D12EnableLightweightDRED, TEXT("dred")))
+	else if (UE::RHI::ShouldEnableGPUCrashFeature(*CVarD3D12EnableLightweightDRED, TEXT("lightdred")))
 	{
 		// Intel suffers a significant performance hit.
 		return IsRHIDeviceIntel()
diff --git a/Engine/Source/Runtime/D3D12RHI/Private/D3D12CommandContext.cpp b/Engine/Source/Runtime/D3D12RHI/Private/D3D12CommandContext.cpp
--- a/Engine/Source/Runtime/D3D12RHI/Private/D3D12CommandContext.cpp
+++ b/Engine/Source/Runtime/D3D12RHI/Private/D3D12CommandContext.cpp
@@ -155,8 +155,8 @@ void FD3D12ContextCommon::BindDiagnosticBuffer(FD3D12RootSignature const* RootSi
 #if WITH_RHI_BREADCRUMBS
 	void FD3D12CommandContext::RHIBeginBreadcrumbGPU(FRHIBreadcrumbNode* Breadcrumb)
 	{
-		// Always emit breadcrumb ID when the diagnostic buffer is available.
-		if (FD3D12DiagnosticBuffer* DiagBuffer = Device->GetQueue(QueueType).DiagnosticBuffer.Get())
+		FD3D12DiagnosticBuffer* DiagBuffer = Device->GetQueue(QueueType).DiagnosticBuffer.Get();
+		if (DiagBuffer && UE::RHI::UseGPUCrashBreadcrumbs())
 		{
 			D3D12_GPU_VIRTUAL_ADDRESS Marker = DiagBuffer->GetGPUQueueMarkerIn();
 			WriteMarker(Marker, Breadcrumb->Name.ID, EMarkerType::In);
@@ -264,8 +264,8 @@ void FD3D12ContextCommon::BindDiagnosticBuffer(FD3D12RootSignature const* RootSi
 		UE::RHICore::Nvidia::Aftermath::D3D12::EndBreadcrumb(AftermathHandle(), Breadcrumb);
 	#endif
 
-		// Always emit breadcrumb ID when the diagnostic buffer is available.
-		if (FD3D12DiagnosticBuffer* DiagBuffer = Device->GetQueue(QueueType).DiagnosticBuffer.Get())
+		FD3D12DiagnosticBuffer* DiagBuffer = Device->GetQueue(QueueType).DiagnosticBuffer.Get();
+		if (DiagBuffer && UE::RHI::UseGPUCrashBreadcrumbs())
 		{
 			D3D12_GPU_VIRTUAL_ADDRESS Marker = DiagBuffer->GetGPUQueueMarkerOut();
 			WriteMarker(Marker, Breadcrumb->Name.ID, EMarkerType::Out);
diff --git a/Engine/Source/Runtime/RHI/Private/RHI.cpp b/Engine/Source/Runtime/RHI/Private/RHI.cpp
--- a/Engine/Source/Runtime/RHI/Private/RHI.cpp
+++ b/Engine/Source/Runtime/RHI/Private/RHI.cpp
@@ -2271,7 +2271,7 @@ namespace UE::RHI
 	static TAutoConsoleVariable<int32> CVarGPUCrashDebuggingBreadcrumbs(
 		TEXT("r.GPUCrashDebugging.Breadcrumbs"),
 		0,
-		TEXT("Enable vendor specific GPU crash analysis tools"),
+		TEXT("Enable RHI breadcrumbs, a vendor-agnostic method for determining which passes were active when a GPU crash occurs"),
 		ECVF_ReadOnly
 	);
 
diff --git a/Engine/Source/Runtime/RHICore/Private/RHICoreNvidiaAftermath.cpp b/Engine/Source/Runtime/RHICore/Private/RHICoreNvidiaAftermath.cpp
--- a/Engine/Source/Runtime/RHICore/Private/RHICoreNvidiaAftermath.cpp
+++ b/Engine/Source/Runtime/RHICore/Private/RHICoreNvidiaAftermath.cpp
@@ -158,7 +158,7 @@ namespace UE::RHICore::Nvidia::Aftermath
 #else
 			GFSDK_Aftermath_GpuCrashDumpWatchedApiFlags_Vulkan
 #endif
-			GFSDK_Aftermath_GpuCrashDumpFeatureFlags_Default,
+			GFSDK_Aftermath_GpuCrashDumpFeatureFlags_DeferDebugInfoCallbacks,
 			&Callback_GpuCrashDump,
 			bDumpShaderDebugInfo ? &Callback_ShaderDebugInfo : nullptr,
 			&Callback_GpuCrashDumpDesc,
