Fix warning C6287: Redundant code.
#rnx

diff --git a/Engine/Source/Runtime/D3D12RHI/Private/D3D12RayTracing.cpp b/Engine/Source/Runtime/D3D12RHI/Private/D3D12RayTracing.cpp
--- a/Engine/Source/Runtime/D3D12RHI/Private/D3D12RayTracing.cpp
+++ b/Engine/Source/Runtime/D3D12RHI/Private/D3D12RayTracing.cpp
@@ -2921,10 +2921,10 @@ void FD3D12RayTracingGeometry::Swap(FD3D12RayTracingGeometry& Other)
 	});
 	::Swap(AccelerationStructureCompactedSize, Other.AccelerationStructureCompactedSize);
 
-	for (uint32 GPUIndex = 0; GPUIndex < MAX_NUM_GPUS && GPUIndex < GNumExplicitGPUsForRendering; ++GPUIndex)
+	FOREACH_GPU(GPUIndex < MAX_NUM_GPUS && GPUIndex < GNumExplicitGPUsForRendering,
 	{
 		UnregisterAsRenameListener(GPUIndex);
-	}
+	});
 
 	Initializer = Other.Initializer;
 
@@ -2935,11 +2935,11 @@ void FD3D12RayTracingGeometry::Swap(FD3D12RayTracingGeometry& Other)
 	GeometryDescs.SetNumUninitialized(Initializer.Segments.Num());
 	TranslateRayTracingGeometryDescs(Initializer, GeometryDescs);
 
-	for (uint32 GPUIndex = 0; GPUIndex < MAX_NUM_GPUS && GPUIndex < GNumExplicitGPUsForRendering; ++GPUIndex)
+	FOREACH_GPU(GPUIndex < MAX_NUM_GPUS && GPUIndex < GNumExplicitGPUsForRendering,
 	{
 		RegisterAsRenameListener(GPUIndex);
 		SetupHitGroupSystemParameters(GPUIndex);
-	}
+	});
 }
 
 void FD3D12RayTracingGeometry::ReleaseUnderlyingResource()
