Use bindless for sampling index and vertex buffers from the hitgroup parameters

diff --git a/Engine/Shaders/Public/Platform/D3D/D3DCommon.ush b/Engine/Shaders/Public/Platform/D3D/D3DCommon.ush
--- a/Engine/Shaders/Public/Platform/D3D/D3DCommon.ush
+++ b/Engine/Shaders/Public/Platform/D3D/D3DCommon.ush
@@ -148,4 +148,27 @@ void UEReportAssertWithPayload(uint AssertID, uint4 Payload)
 	#define GetResourceFromHeap(Type, Index) ResourceDescriptorHeap[Index]
 	#define GetSamplerFromHeap(Type, Index) SamplerDescriptorHeap[Index]
 #endif
+	
 
+#if ENABLE_BINDLESS_RESOURCES && PLATFORM_SUPPORTS_BINDLESS && (RAYCALLABLESHADER || RAYHITGROUPSHADER || RAYGENSHADER || RAYMISSSHADER)
+	#define OVERRIDE_RAY_TRACING_HIT_GROUP_SYSTEM_RESOURCES 1
+	#include "/Engine/Shared/RayTracingBuiltInResources.h"
+
+	// Combine the common FHitGroupSystemRootConstants with bindless parameters
+	struct FD3DHitGroupSystemParameters
+	{
+		FHitGroupSystemRootConstants RootConstants;
+
+		uint BindlessHitGroupSystemIndexBuffer;
+		uint BindlessHitGroupSystemVertexBuffer;
+	};
+
+	ConstantBuffer<FD3DHitGroupSystemParameters>	D3DHitGroupSystemParameters : UE_HLSL_REGISTER(b, RAY_TRACING_SYSTEM_ROOTCONSTANT_REGISTER, UE_HLSL_SPACE_RAY_TRACING_SYSTEM);
+	
+	static const FHitGroupSystemRootConstants HitGroupSystemRootConstants = D3DHitGroupSystemParameters.RootConstants;
+
+	// Will get filled in with the bindless index in D3DHitGroupSystemParameters.BindlessHitGroupSystemIndexBuffer
+	ByteAddressBuffer HitGroupSystemIndexBuffer;
+	// Will get filled in with the bindless index in D3DHitGroupSystemParameters.BindlessHitGroupSystemVertexBuffer
+	ByteAddressBuffer HitGroupSystemVertexBuffer;
+#endif
diff --git a/Engine/Source/Developer/Windows/ShaderFormatD3D/Private/D3DShaderCompiler.cpp b/Engine/Source/Developer/Windows/ShaderFormatD3D/Private/D3DShaderCompiler.cpp
--- a/Engine/Source/Developer/Windows/ShaderFormatD3D/Private/D3DShaderCompiler.cpp
+++ b/Engine/Source/Developer/Windows/ShaderFormatD3D/Private/D3DShaderCompiler.cpp
@@ -1322,13 +1322,29 @@ bool CompileAndProcessD3DShaderFXC(
 
 struct FD3DShaderParameterParserPlatformConfiguration : public FShaderParameterParser::FPlatformConfiguration
 {
-	FD3DShaderParameterParserPlatformConfiguration()
+	FD3DShaderParameterParserPlatformConfiguration(const FShaderCompilerInput& Input)
 		: FShaderParameterParser::FPlatformConfiguration(TEXTVIEW("cbuffer"), EShaderParameterParserConfigurationFlags::UseStableConstantBuffer|EShaderParameterParserConfigurationFlags::SupportsBindless)
+		, bIsRayTracingShader(Input.IsRayTracingShader())
+		, HitGroupSystemIndexBufferName(FShaderParameterParser::kBindlessSRVPrefix + FString(TEXT("HitGroupSystemIndexBuffer")))
+		, HitGroupSystemVertexBufferName(FShaderParameterParser::kBindlessSRVPrefix + FString(TEXT("HitGroupSystemVertexBuffer")))
 	{
 	}
 
 	virtual FString GenerateBindlessAccess(EBindlessConversionType BindlessType, FStringView FullTypeString, FStringView ArrayNameOverride, FStringView IndexString) const final
 	{
+		if (bIsRayTracingShader && (BindlessType == EBindlessConversionType::SRV))
+		{
+			// Patch the HitGroupSystemIndexBuffer/HitGroupSystemVertexBuffer indices to use the ones contained in the shader record
+			if (IndexString == HitGroupSystemIndexBufferName)
+			{
+				IndexString = TEXTVIEW("D3DHitGroupSystemParameters.BindlessHitGroupSystemIndexBuffer");
+			}
+			else if (IndexString == HitGroupSystemVertexBufferName)
+			{
+				IndexString = TEXTVIEW("D3DHitGroupSystemParameters.BindlessHitGroupSystemVertexBuffer");
+			}
+		}
+
 		// GetResourceFromHeap(Type, Index) ResourceDescriptorHeap[Index]
 		// GetSamplerFromHeap(Type, Index)  SamplerDescriptorHeap[Index]
 
@@ -1339,6 +1355,10 @@ struct FD3DShaderParameterParserPlatformConfiguration : public FShaderParameterP
 			IndexString.Len(), IndexString.GetData()
 		);
 	}
+	
+	const bool bIsRayTracingShader;
+	const FString HitGroupSystemIndexBufferName;
+	const FString HitGroupSystemVertexBufferName;
 };
 
 void CompileD3DShader(const FShaderCompilerInput& Input, const FShaderPreprocessOutput& InPreprocessOutput, FShaderCompilerOutput& Output, const FString& WorkingDirectory, ED3DShaderModel ShaderModel)
@@ -1356,7 +1376,7 @@ void CompileD3DShader(const FShaderCompilerInput& Input, const FShaderPreprocess
 	FString EntryPointName = Input.EntryPointName;
 	FString PreprocessedSource(InPreprocessOutput.GetSourceViewWide());
 
-	FD3DShaderParameterParserPlatformConfiguration PlatformConfiguration;
+	FD3DShaderParameterParserPlatformConfiguration PlatformConfiguration(Input);
 	FShaderParameterParser ShaderParameterParser(PlatformConfiguration);
 	if (!ShaderParameterParser.ParseAndModify(Input, Output.Errors, PreprocessedSource))
 	{
diff --git a/Engine/Source/Runtime/D3D12RHI/Internal/D3D12RayTracingResources.h b/Engine/Source/Runtime/D3D12RHI/Internal/D3D12RayTracingResources.h
new file mode 100644
--- /dev/null
+++ b/Engine/Source/Runtime/D3D12RHI/Internal/D3D12RayTracingResources.h
@@ -0,0 +1,38 @@
+// Copyright Epic Games, Inc. All Rights Reserved.
+
+/*=============================================================================
+	D3D12ShaderResources.h: Internal D3D12 RHI RayTracing definitions.
+=============================================================================*/
+
+#pragma once
+
+#include "HAL/Platform.h"
+#include "RayTracingBuiltInResources.h"
+
+#include "D3D12RHI.h"
+
+#if !defined(D3D12_MAJOR_VERSION)
+	#include "D3D12ThirdParty.h"
+#endif
+
+// Built-in local root parameters that are always bound to all hit shaders
+// Contains union for bindless and non-bindless index/vertex buffer data to make the code handling the hit group parameters easier to use
+// (otherwise all cached hit parameter code has to be done twice and stored twice making everything more complicated)
+// Ideally the non bindless code path should be removed 'soon' - this constant buffer size for FD3D12HitGroupSystemParameters in Bindless is 8 bytes bigger than needed
+struct FD3D12HitGroupSystemParameters
+{	
+	FHitGroupSystemRootConstants RootConstants;
+	union
+	{
+		struct
+		{
+			uint32 BindlessHitGroupSystemIndexBuffer;
+			uint32 BindlessHitGroupSystemVertexBuffer;
+		};
+		struct
+		{
+			D3D12_GPU_VIRTUAL_ADDRESS IndexBuffer;
+			D3D12_GPU_VIRTUAL_ADDRESS VertexBuffer;
+		};
+	};
+};
diff --git a/Engine/Source/Runtime/D3D12RHI/Private/D3D12RayTracing.cpp b/Engine/Source/Runtime/D3D12RHI/Private/D3D12RayTracing.cpp
--- a/Engine/Source/Runtime/D3D12RHI/Private/D3D12RayTracing.cpp
+++ b/Engine/Source/Runtime/D3D12RHI/Private/D3D12RayTracing.cpp
@@ -1244,6 +1244,18 @@ private:
 	FD3D12RootSignature DefaultLocalRootSignature; // Default empty root signature used for default hit shaders.
 };
 
+inline bool AreBindlessResourcesEnabled(FD3D12Adapter* Adapter)
+{
+#if PLATFORM_SUPPORTS_BINDLESS_RENDERING
+	FD3D12BindlessDescriptorManager& Manager = Adapter->GetDevice(0)->GetBindlessDescriptorManager();
+	if (Manager.AreResourcesBindless())
+	{
+		return true;
+	}
+#endif
+	return false;
+}
+
 // Helper class used to manage SBT buffer for a specific GPU
 class FD3D12RayTracingShaderBindingTableInternal
 {
@@ -1284,7 +1296,7 @@ public:
 		checkf(Initializer.LocalBindingDataSize <= 4096, TEXT("The maximum size of a local root signature is 4KB.")); // as per section 4.22.1 of DXR spec v1.0
 
 		const uint32 NumHitGroupSlots = Initializer.bAllowHitGroupIndexing ? Initializer.NumGeometrySegments * Initializer.NumShaderSlotsPerGeometrySegment : 1;
-		checkf(Initializer.LocalBindingDataSize >= sizeof(FHitGroupSystemParameters), TEXT("All local root signatures are expected to contain ray tracing system root parameters (2x root buffers + 4x root DWORD)"));
+		checkf(Initializer.LocalBindingDataSize >= sizeof(FD3D12HitGroupSystemParameters), TEXT("All local root signatures are expected to contain ray tracing system root parameters"));
 		
 		bAllowHitGroupIndexing = Initializer.bAllowHitGroupIndexing;
 		LocalRecordSizeUnaligned = ShaderIdentifierSize + Initializer.LocalBindingDataSize;
@@ -1426,7 +1438,7 @@ public:
 		WriteData(WriteOffset, ShaderIdentifier.Data, ShaderIdentifierSize);
 	}
 
-	void SetHitGroupSystemParameters(uint32 RecordIndex, const FHitGroupSystemParameters& SystemParameters)
+	void SetHitGroupSystemParameters(uint32 RecordIndex, const FD3D12HitGroupSystemParameters& SystemParameters)
 	{
 		const uint32 OffsetWithinRootSignature = 0; // System parameters are always first in the RS.
 		const uint32 ShaderTableOffset = HitGroupShaderTableOffset;
@@ -3000,7 +3012,7 @@ void FD3D12RayTracingGeometry::ReleaseUnderlyingResource()
 
 	AccelerationStructureCompactedSize = 0;
 	GeometryDescs = {};
-	for (TArray<FHitGroupSystemParameters>& HitGroupParametersForGPU : HitGroupSystemParameters)
+	for (TArray<FD3D12HitGroupSystemParameters>& HitGroupParametersForGPU : HitGroupSystemParameters)
 	{
 		HitGroupParametersForGPU.Empty();
 	}
@@ -3011,8 +3023,76 @@ FD3D12RayTracingGeometry::~FD3D12RayTracingGeometry()
 	ReleaseUnderlyingResource();
 }
 
+void FD3D12RayTracingGeometry::AllocateBufferSRVs(uint32 InGPUIndex)
+{
+	HitGroupSystemIndexBufferSRV[InGPUIndex].Reset();
+	HitGroupSystemSegmentVertexBufferSRVs[InGPUIndex].Empty();
+
+	// Procedural doesn't need any SRVs for index buffer
+	if (Initializer.IndexBuffer && Initializer.GeometryType == RTGT_Triangles)
+	{
+		check(Initializer.IndexBufferOffset == 0);
+
+		FD3D12Buffer* IndexBuffer = FD3D12DynamicRHI::ResourceCast(Initializer.IndexBuffer.GetReference());
+
+		D3D12_SHADER_RESOURCE_VIEW_DESC SRVDesc = {};
+		SRVDesc.ViewDimension = D3D12_SRV_DIMENSION_BUFFER;
+		SRVDesc.Format = DXGI_FORMAT_R32_TYPELESS;
+		SRVDesc.Shader4ComponentMapping = D3D12_DEFAULT_SHADER_4_COMPONENT_MAPPING;
+		SRVDesc.Buffer.FirstElement = (Initializer.IndexBufferOffset + IndexBuffer->ResourceLocation.GetOffsetFromBaseOfResource()) >> 2u;
+		SRVDesc.Buffer.NumElements = IndexBuffer->GetSize() >> 2u;
+		SRVDesc.Buffer.Flags = D3D12_BUFFER_SRV_FLAG_RAW;
+		SRVDesc.Buffer.StructureByteStride = 0;
+
+		HitGroupSystemIndexBufferSRV[InGPUIndex] = MakeShared<FD3D12ShaderResourceView>(GetParentAdapter()->GetDevice(InGPUIndex), InGPUIndex > 0 ? HitGroupSystemIndexBufferSRV[0].Get() : nullptr);
+		HitGroupSystemIndexBufferSRV[InGPUIndex]->CreateView(IndexBuffer, SRVDesc, FD3D12ShaderResourceView::EFlags::None);
+	}
+
+	for (const FRayTracingGeometrySegment& Segment : Initializer.Segments)
+	{
+		checkf((Segment.VertexBufferOffset % 16) == 0, TEXT("The byte offset of raw views must be a multiple of 16 (specified offset: %d)."), Segment.VertexBufferOffset);
+
+		FD3D12Buffer* VertexBuffer = FD3D12DynamicRHI::ResourceCast(Segment.VertexBuffer.GetReference());
+
+		D3D12_SHADER_RESOURCE_VIEW_DESC SRVDesc = {};
+		SRVDesc.ViewDimension = D3D12_SRV_DIMENSION_BUFFER;
+		SRVDesc.Format = DXGI_FORMAT_R32_TYPELESS;
+		SRVDesc.Shader4ComponentMapping = D3D12_DEFAULT_SHADER_4_COMPONENT_MAPPING;
+		SRVDesc.Buffer.FirstElement = (Segment.VertexBufferOffset + VertexBuffer->ResourceLocation.GetOffsetFromBaseOfResource()) >> 2u;
+		if (Initializer.GeometryType == RTGT_Procedural)
+		{
+			SRVDesc.Buffer.NumElements = Segment.NumPrimitives * Segment.VertexBufferStride / 4; //< NumElements in R32 size
+		}
+		else
+		{
+			SRVDesc.Buffer.NumElements = Segment.MaxVertices * Segment.VertexBufferStride / 4; //< NumElements in R32 size
+		}
+		SRVDesc.Buffer.Flags = D3D12_BUFFER_SRV_FLAG_RAW;
+		SRVDesc.Buffer.StructureByteStride = 0;
+
+		FD3D12ShaderResourceView* FirstLinkedObject = nullptr;
+		if (InGPUIndex > 0)
+		{
+			int32 SegmentIndex = HitGroupSystemSegmentVertexBufferSRVs[InGPUIndex].Num();
+			if (HitGroupSystemSegmentVertexBufferSRVs[0].Num() > SegmentIndex)
+			{
+				FirstLinkedObject = HitGroupSystemSegmentVertexBufferSRVs[0][SegmentIndex].Get();
+			}
+		}
+		TSharedPtr<FD3D12ShaderResourceView> VertexBufferSRV = MakeShared<FD3D12ShaderResourceView>(GetParentAdapter()->GetDevice(InGPUIndex), FirstLinkedObject);
+		VertexBufferSRV->CreateView(VertexBuffer, SRVDesc, FD3D12ShaderResourceView::EFlags::None);
+		HitGroupSystemSegmentVertexBufferSRVs[InGPUIndex].Add(VertexBufferSRV);
+	}
+}
+
 void FD3D12RayTracingGeometry::RegisterAsRenameListener(uint32 InGPUIndex)
 {
+	// Not needed if bindless
+	if (AreBindlessResourcesEnabled(GetParentAdapter()))
+	{
+		return;
+	}
+	
 	check(!bRegisteredAsRenameListener[InGPUIndex]);
 
 	FD3D12Buffer* IndexBuffer = FD3D12DynamicRHI::ResourceCast(Initializer.IndexBuffer.GetReference(), InGPUIndex);
@@ -3043,6 +3123,8 @@ void FD3D12RayTracingGeometry::UnregisterAsRenameListener(uint32 InGPUIndex)
 		return;
 	}
 
+	check(!AreBindlessResourcesEnabled(GetParentAdapter()));
+
 	FD3D12Buffer* IndexBuffer = FD3D12DynamicRHI::ResourceCast(Initializer.IndexBuffer.GetReference(), InGPUIndex);
 	if (IndexBuffer)
 	{
@@ -3066,6 +3148,8 @@ void FD3D12RayTracingGeometry::UnregisterAsRenameListener(uint32 InGPUIndex)
 
 void FD3D12RayTracingGeometry::ResourceRenamed(FD3D12ContextArray const& Contexts, FD3D12BaseShaderResource* InRenamedResource, FD3D12ResourceLocation* InNewResourceLocation)
 {
+	check(!AreBindlessResourcesEnabled(GetParentAdapter()));
+
 	// Empty resource location is used on destruction of the base shader resource but this
 	// shouldn't happen for RT Geometries because it keeps smart pointers to it's resources.
 	check(InNewResourceLocation != nullptr);
@@ -3157,23 +3241,45 @@ void FD3D12RayTracingGeometry::SetupHitGroupSystemParameters(uint32 InGPUIndex)
 {
 	D3D12_RAYTRACING_GEOMETRY_TYPE GeometryType = TranslateRayTracingGeometryType(Initializer.GeometryType);
 
-	TArray<FHitGroupSystemParameters>& HitGroupSystemParametersForThisGPU = HitGroupSystemParameters[InGPUIndex];
+	bool bBindless = AreBindlessResourcesEnabled(GetParentAdapter());
+
+	TArray<FD3D12HitGroupSystemParameters>& HitGroupSystemParametersForThisGPU = HitGroupSystemParameters[InGPUIndex];
 	HitGroupSystemParametersForThisGPU.Reset(Initializer.Segments.Num());
 
 	check(BuffersValid(InGPUIndex));
+	if (bBindless)
+	{
+		AllocateBufferSRVs(InGPUIndex);
+	}
+
 	FD3D12Buffer* IndexBuffer = FD3D12DynamicRHI::ResourceCast(Initializer.IndexBuffer.GetReference(), InGPUIndex);
 	const uint32 IndexStride = IndexBuffer ? IndexBuffer->GetStride() : 0;
-	for (const FRayTracingGeometrySegment& Segment : Initializer.Segments)
+	for (int32 SegmentIndex = 0; SegmentIndex < Initializer.Segments.Num(); ++SegmentIndex)
 	{
+		const FRayTracingGeometrySegment& Segment = Initializer.Segments[SegmentIndex];
 		FD3D12Buffer* VertexBuffer = FD3D12DynamicRHI::ResourceCast(Segment.VertexBuffer.GetReference(), InGPUIndex);
 
-		FHitGroupSystemParameters SystemParameters = {};
+		FD3D12HitGroupSystemParameters SystemParameters = {};
 		SystemParameters.RootConstants.SetVertexAndIndexStride(Segment.VertexBufferStride, IndexStride);
-		SystemParameters.VertexBuffer = VertexBuffer->ResourceLocation.GetGPUVirtualAddress() + Segment.VertexBufferOffset;
+		if (bBindless)
+		{
+			SystemParameters.BindlessHitGroupSystemVertexBuffer = HitGroupSystemSegmentVertexBufferSRVs[InGPUIndex][SegmentIndex]->GetBindlessHandle().GetIndex();
+		}
+		else
+		{
+			SystemParameters.VertexBuffer = VertexBuffer->ResourceLocation.GetGPUVirtualAddress() + Segment.VertexBufferOffset;
+		}
 
 		if (GeometryType == D3D12_RAYTRACING_GEOMETRY_TYPE_TRIANGLES && IndexBuffer != nullptr)
 		{
-			SystemParameters.IndexBuffer = IndexBuffer->ResourceLocation.GetGPUVirtualAddress();
+			if (bBindless)
+			{
+				SystemParameters.BindlessHitGroupSystemIndexBuffer = HitGroupSystemIndexBufferSRV[InGPUIndex]->GetBindlessHandle().GetIndex();
+			}
+			else
+			{
+				SystemParameters.IndexBuffer = IndexBuffer->ResourceLocation.GetGPUVirtualAddress();
+			}
 			SystemParameters.RootConstants.IndexBufferOffsetInBytes = Initializer.IndexBufferOffset + IndexStride * Segment.FirstPrimitive * FD3D12RayTracingGeometry::IndicesPerPrimitive;
 			SystemParameters.RootConstants.FirstPrimitive = Segment.FirstPrimitive;
 		}
@@ -3199,8 +3305,8 @@ void FD3D12RayTracingGeometry::CreateAccelerationStructureBuildDesc(FD3D12Comman
 
 	FD3D12Buffer* IndexBuffer = CommandContext.RetrieveObject<FD3D12Buffer>(Initializer.IndexBuffer.GetReference());
 	FD3D12Buffer* NullTransformBufferD3D12 = CommandContext.RetrieveObject<FD3D12Buffer>(NullTransformBuffer.GetReference());
-
-	const TArray<FHitGroupSystemParameters>& HitGroupSystemParametersForThisGPU = HitGroupSystemParameters[GPUIndex];
+		
+	const TArray<FD3D12HitGroupSystemParameters>& HitGroupSystemParametersForThisGPU = HitGroupSystemParameters[GPUIndex];
 	check(HitGroupSystemParametersForThisGPU.Num() == Initializer.Segments.Num());
 
 	ERayTracingAccelerationStructureFlags BuildFlags = GetRayTracingAccelerationStructureBuildFlags(Initializer);
@@ -3211,8 +3317,8 @@ void FD3D12RayTracingGeometry::CreateAccelerationStructureBuildDesc(FD3D12Comman
 		Desc = GeometryDescs[SegmentIndex]; // Copy from template
 
 		const FRayTracingGeometrySegment& Segment = Initializer.Segments[SegmentIndex];
-		const FHitGroupSystemParameters& SystemParameters = HitGroupSystemParametersForThisGPU[SegmentIndex];
-
+		const FD3D12HitGroupSystemParameters& SystemParameters = HitGroupSystemParametersForThisGPU[SegmentIndex];
+		
 		FD3D12Buffer* VertexBuffer = CommandContext.RetrieveObject<FD3D12Buffer>(Segment.VertexBuffer.GetReference());
 
 		switch (GeometryType)
@@ -3250,7 +3356,7 @@ void FD3D12RayTracingGeometry::CreateAccelerationStructureBuildDesc(FD3D12Comman
 				check(Desc.Triangles.IndexCount <= Segment.NumPrimitives * FD3D12RayTracingGeometry::IndicesPerPrimitive);
 
 				Desc.Triangles.IndexFormat = (IndexStride == 4 ? DXGI_FORMAT_R32_UINT : DXGI_FORMAT_R16_UINT);
-				Desc.Triangles.IndexBuffer = SystemParameters.IndexBuffer + SystemParameters.RootConstants.IndexBufferOffsetInBytes;
+				Desc.Triangles.IndexBuffer = IndexBuffer->ResourceLocation.GetGPUVirtualAddress() + SystemParameters.RootConstants.IndexBufferOffsetInBytes;
 			}
 			else
 			{
@@ -3261,13 +3367,13 @@ void FD3D12RayTracingGeometry::CreateAccelerationStructureBuildDesc(FD3D12Comman
 				check(Desc.Triangles.IndexBuffer == D3D12_GPU_VIRTUAL_ADDRESS(0));
 			}
 
-			Desc.Triangles.VertexBuffer.StartAddress = SystemParameters.VertexBuffer;
+			Desc.Triangles.VertexBuffer.StartAddress = VertexBuffer->ResourceLocation.GetGPUVirtualAddress() + Segment.VertexBufferOffset;
 			Desc.Triangles.VertexBuffer.StrideInBytes = Segment.VertexBufferStride;
 			break;
 
 		case D3D12_RAYTRACING_GEOMETRY_TYPE_PROCEDURAL_PRIMITIVE_AABBS:
 			Desc.AABBs.AABBCount = Segment.NumPrimitives;
-			Desc.AABBs.AABBs.StartAddress = SystemParameters.VertexBuffer;
+			Desc.AABBs.AABBs.StartAddress = VertexBuffer->ResourceLocation.GetGPUVirtualAddress() + Segment.VertexBufferOffset;
 			Desc.AABBs.AABBs.StrideInBytes = Segment.VertexBufferStride;
 			break;
 
@@ -4908,7 +5014,7 @@ static void SetRayTracingHitGroup(
 	}
 #endif // DO_CHECK
 
-	FHitGroupSystemParameters SystemParameters = Geometry->HitGroupSystemParameters[GPUIndex][GeometrySegmentIndex];
+	FD3D12HitGroupSystemParameters SystemParameters = Geometry->HitGroupSystemParameters[GPUIndex][GeometrySegmentIndex];
 	SystemParameters.RootConstants.UserData = UserData;
 
 	ShaderTable->SetHitGroupSystemParameters(RecordIndex, SystemParameters);
@@ -4931,7 +5037,7 @@ static void SetRayTracingHitGroup(
 		if (ExistingRecordIndex)
 		{
 			// Simply copy local shader parameters from existing SBT record and set the shader identifier, skipping resource binding work.
-			const uint32 OffsetFromRootSignatureStart = sizeof(FHitGroupSystemParameters);
+			const uint32 OffsetFromRootSignatureStart = sizeof(FD3D12HitGroupSystemParameters);
 			ShaderTable->SetHitGroupIdentifier(RecordIndex, Pipeline->HitGroupShaders.Identifiers[HitGroupIndex]);
 			ShaderTable->CopyHitGroupParameters(RecordIndex, *ExistingRecordIndex, OffsetFromRootSignatureStart);
 			return;
@@ -4971,7 +5077,7 @@ static void SetRayTracingCallableShader(
 {
 	checkf(RecordIndex < ShaderTable->NumCallableRecords, TEXT("Callable shader record index is invalid. Make sure that NumCallableShaderSlots is correct in FRayTracingShaderBindingTableInitializer."));
 
-	const uint32 UserDataOffset = offsetof(FHitGroupSystemParameters, RootConstants) + offsetof(FHitGroupSystemRootConstants, UserData);
+	const uint32 UserDataOffset = offsetof(FD3D12HitGroupSystemParameters, RootConstants) + offsetof(FHitGroupSystemRootConstants, UserData);
 	ShaderTable->SetCallableShaderParameters(RecordIndex, UserDataOffset, UserData);
 
 	const FD3D12ShaderIdentifier* ShaderIdentifier = &FD3D12ShaderIdentifier::Null;
@@ -5011,7 +5117,7 @@ static void SetRayTracingMissShader(
 {
 	checkf(RecordIndex < ShaderTable->NumMissRecords, TEXT("Miss shader record index is invalid. Make sure that NumMissShaderSlots is correct in FRayTracingShaderBindingTableInitializer."));
 
-	const uint32 UserDataOffset = offsetof(FHitGroupSystemParameters, RootConstants) + offsetof(FHitGroupSystemRootConstants, UserData);
+	const uint32 UserDataOffset = offsetof(FD3D12HitGroupSystemParameters, RootConstants) + offsetof(FHitGroupSystemRootConstants, UserData);
 	ShaderTable->SetMissShaderParameters(RecordIndex, UserDataOffset, UserData);
 
 	const FD3D12RayTracingShader* Shader = Pipeline->MissShaders.Shaders[ShaderIndexInPipeline];
diff --git a/Engine/Source/Runtime/D3D12RHI/Private/D3D12RayTracing.h b/Engine/Source/Runtime/D3D12RHI/Private/D3D12RayTracing.h
--- a/Engine/Source/Runtime/D3D12RHI/Private/D3D12RayTracing.h
+++ b/Engine/Source/Runtime/D3D12RHI/Private/D3D12RayTracing.h
@@ -6,19 +6,11 @@
 
 #if D3D12_RHI_RAYTRACING
 
-#include "RayTracingBuiltInResources.h"
+#include "D3D12RayTracingResources.h"
 
 class FD3D12RayTracingPipelineState;
 class FD3D12RayTracingShaderBindingTable;
 
-// Built-in local root parameters that are always bound to all hit shaders
-struct FHitGroupSystemParameters
-{
-	D3D12_GPU_VIRTUAL_ADDRESS IndexBuffer;
-	D3D12_GPU_VIRTUAL_ADDRESS VertexBuffer;
-	FHitGroupSystemRootConstants RootConstants;
-};
-
 class FD3D12RayTracingGeometry : public FRHIRayTracingGeometry, public FD3D12AdapterChild, public FD3D12ShaderResourceRenameListener, public FNoncopyable
 {
 public:
@@ -43,6 +35,8 @@ public:
 	// Implement FD3D12ShaderResourceRenameListener interface
 	virtual void ResourceRenamed(FD3D12ContextArray const& Contexts, FD3D12BaseShaderResource* InRenamedResource, FD3D12ResourceLocation* InNewResourceLocation) override;
 
+	void AllocateBufferSRVs(uint32 InGPUIndex);
+
 	void RegisterAsRenameListener(uint32 InGPUIndex);
 	void UnregisterAsRenameListener(uint32 InGPUIndex);
 
@@ -77,7 +71,11 @@ public:
 	bool bHasPendingCompactionRequests[MAX_NUM_GPUS];
 
 	// Hit shader parameters per geometry segment
-	TArray<FHitGroupSystemParameters> HitGroupSystemParameters[MAX_NUM_GPUS];
+	TArray<FD3D12HitGroupSystemParameters> HitGroupSystemParameters[MAX_NUM_GPUS];
+
+	// RAW SRVs to index and vertex buffers when using bindless hit group paramaters
+	TSharedPtr<FD3D12ShaderResourceView> HitGroupSystemIndexBufferSRV[MAX_NUM_GPUS];
+	TArray<TSharedPtr<FD3D12ShaderResourceView>> HitGroupSystemSegmentVertexBufferSRVs[MAX_NUM_GPUS];
 
 	FDebugName DebugName;
 	FName OwnerName;		// Store the path name of the owner object for resource tracking
diff --git a/Engine/Source/Runtime/D3D12RHI/Private/D3D12RootSignature.cpp b/Engine/Source/Runtime/D3D12RHI/Private/D3D12RootSignature.cpp
--- a/Engine/Source/Runtime/D3D12RHI/Private/D3D12RootSignature.cpp
+++ b/Engine/Source/Runtime/D3D12RHI/Private/D3D12RootSignature.cpp
@@ -9,6 +9,7 @@
 #include "D3D12AmdExtensions.h"
 #include "D3D12RootSignatureDefinitions.h"
 #include "RayTracingBuiltInResources.h"
+#include "D3D12RayTracing.h"
 
 #ifndef FD3D12_ROOT_SIGNATURE_FLAG_GLOBAL_ROOT_SIGNATURE
 #define FD3D12_ROOT_SIGNATURE_FLAG_GLOBAL_ROOT_SIGNATURE D3D12_ROOT_SIGNATURE_FLAG_NONE
@@ -169,47 +170,66 @@ FD3D12RootSignatureDesc::FD3D12RootSignatureDesc(const FD3D12QuantizedBoundShade
 	{
 		BindingSpace = UE_HLSL_SPACE_RAY_TRACING_LOCAL;
 
-		// Add standard root parameters for hit groups, as per FHitGroupSystemParameters declaration in D3D12RayTracing.cpp and RayTracingHitGroupCommon.ush:
-		//          Resources:
-		// 8 bytes: index buffer as root SRV (raw buffer)
-		// 8 bytes: vertex buffer as root SRV (raw buffer)
+		// Add standard root parameters for hit groups, as per FD3D12HitGroupSystemParameters declaration in D3D12Util.h 
+		// and RayTracingHitGroupCommon.ush (Non-Bindless) and D3DCommon.ush (Bindless):
 		//          FHitGroupSystemRootConstants:
 		// 4 bytes: index/vertex fetch configuration as root constant (bitfield defining index and vertex formats)
 		// 4 bytes: index buffer offset in bytes
 		// 4 bytes: first primitive of the segment (as set in FRayTracingGeometrySegment)
 		// 4 bytes: hit group user data
-		// 4 bytes: index of the first instance that belongs to the current batch.
-		// 4 bytes: unused padding to ensure the next parameter is aligned to 8-byte boundary
+		//          Non-Bindless Resources:
+		// 8 bytes: index buffer as root SRV (raw buffer)
+		// 8 bytes: vertex buffer as root SRV (raw buffer)
+		//          Bindless Resources:
+		// 4 bytes: index buffer bindless resource index (raw buffer)
+		// 4 bytes: vertex buffer bindless resource index (raw buffer)
+		// 8 bytes: union padding with non-bindless data
+
 		// -----------
-		// 40 bytes
+		// 32 bytes
 
 		check(RootParameterCount == 0 && RootParametersSize == 0); // We expect system RT parameters to come first
-
-		// Index buffer descriptor
-		{
+		
+		// Bindless FD3D12HitGroupSystemParameters structure
+		if (QBSS.bUseDirectlyIndexedResourceHeap)
+		{			
 			check(RootParameterCount < MaxRootParameters);
-			TableSlots[RootParameterCount].InitAsShaderResourceView(RAY_TRACING_SYSTEM_INDEXBUFFER_REGISTER, UE_HLSL_SPACE_RAY_TRACING_SYSTEM);
-			RootParameterCount++;
-			RootParametersSize += RootDescriptorCost;
-		}
+			static_assert(sizeof(FD3D12HitGroupSystemParameters) % 8 == 0, "FD3D12HitGroupSystemParameters structure must be 8-byte aligned");
 
-		// Vertex buffer descriptor
-		{
-			check(RootParameterCount < MaxRootParameters);
-			TableSlots[RootParameterCount].InitAsShaderResourceView(RAY_TRACING_SYSTEM_VERTEXBUFFER_REGISTER, UE_HLSL_SPACE_RAY_TRACING_SYSTEM);
-			RootParameterCount++;
-			RootParametersSize += RootDescriptorCost;
-		}
-
-		// FHitGroupSystemRootConstants structure
-		{
-			check(RootParameterCount < MaxRootParameters);
-			static_assert(sizeof(FHitGroupSystemRootConstants) % 8 == 0, "FHitGroupSystemRootConstants structure must be 8-byte aligned");
-			const uint32 NumConstants = sizeof(FHitGroupSystemRootConstants) / sizeof(uint32);
+			// Num constants could be 2 smaller if non-bindless data is removed (24 bytes instead of 32)
+			const uint32 NumConstants = sizeof(FD3D12HitGroupSystemParameters) / sizeof(uint32);
 			TableSlots[RootParameterCount].InitAsConstants(NumConstants, RAY_TRACING_SYSTEM_ROOTCONSTANT_REGISTER, UE_HLSL_SPACE_RAY_TRACING_SYSTEM);
 			RootParameterCount++;
 			RootParametersSize += NumConstants * RootConstantCost;
 		}
+		else
+		{
+			// FHitGroupSystemRootConstants structure
+			{
+				check(RootParameterCount < MaxRootParameters);
+				static_assert(sizeof(FHitGroupSystemRootConstants) % 8 == 0, "FHitGroupSystemRootConstants structure must be 8-byte aligned");
+				const uint32 NumConstants = sizeof(FHitGroupSystemRootConstants) / sizeof(uint32);
+				TableSlots[RootParameterCount].InitAsConstants(NumConstants, RAY_TRACING_SYSTEM_ROOTCONSTANT_REGISTER, UE_HLSL_SPACE_RAY_TRACING_SYSTEM);
+				RootParameterCount++;
+				RootParametersSize += NumConstants * RootConstantCost;
+			}
+
+			// Index buffer descriptor
+			{
+				check(RootParameterCount < MaxRootParameters);
+				TableSlots[RootParameterCount].InitAsShaderResourceView(RAY_TRACING_SYSTEM_INDEXBUFFER_REGISTER, UE_HLSL_SPACE_RAY_TRACING_SYSTEM);
+				RootParameterCount++;
+				RootParametersSize += RootDescriptorCost;
+			}
+
+			// Vertex buffer descriptor
+			{
+				check(RootParameterCount < MaxRootParameters);
+				TableSlots[RootParameterCount].InitAsShaderResourceView(RAY_TRACING_SYSTEM_VERTEXBUFFER_REGISTER, UE_HLSL_SPACE_RAY_TRACING_SYSTEM);
+				RootParameterCount++;
+				RootParametersSize += RootDescriptorCost;
+			}
+		}	
 	}
 	else if (QBSS.RootSignatureType == RS_RayTracingGlobal)
 	{
@@ -795,4 +815,4 @@ FD3D12QuantizedBoundShaderState FD3D12RootSignatureManager::GetQuantizedBoundSha
 	check(QBSS);
 
 	return *QBSS;
-}
+}
\ No newline at end of file
diff --git a/Engine/Source/Runtime/D3D12RHI/Public/D3D12RootSignatureDefinitions.h b/Engine/Source/Runtime/D3D12RHI/Public/D3D12RootSignatureDefinitions.h
--- a/Engine/Source/Runtime/D3D12RHI/Public/D3D12RootSignatureDefinitions.h
+++ b/Engine/Source/Runtime/D3D12RHI/Public/D3D12RootSignatureDefinitions.h
@@ -3,6 +3,7 @@
 #pragma once
 
 #include "D3D12RHI.h"
+#include "D3D12RayTracingResources.h"
 #include "Containers/Array.h"
 #include "Containers/UnrealString.h"
 #include "Misc/EnumClassFlags.h"
@@ -467,11 +468,21 @@ namespace D3D12ShaderUtils
 
 		if (bLocalRootSignature)
 		{
-			Creator.AddShaderResourceViewParameter(RAY_TRACING_SYSTEM_INDEXBUFFER_REGISTER, UE_HLSL_SPACE_RAY_TRACING_SYSTEM);
-			Creator.AddShaderResourceViewParameter(RAY_TRACING_SYSTEM_VERTEXBUFFER_REGISTER, UE_HLSL_SPACE_RAY_TRACING_SYSTEM);
+			if (Creator.HasFlags(EShaderBindingLayoutFlags::BindlessResources))
+			{
+				// Num constants could be 2 smaller if non-bindless data is removed (24 bytes instead of 32)
+				// (see notes in FD3D12RootSignatureDesc constructor)
+				uint32 NumConstants = sizeof(FD3D12HitGroupSystemParameters) / sizeof(uint32);
+				Creator.AddConstantsParameter(NumConstants, RAY_TRACING_SYSTEM_ROOTCONSTANT_REGISTER, UE_HLSL_SPACE_RAY_TRACING_SYSTEM);
+			}
+			else
+			{
+				Creator.AddShaderResourceViewParameter(RAY_TRACING_SYSTEM_INDEXBUFFER_REGISTER, UE_HLSL_SPACE_RAY_TRACING_SYSTEM);
+				Creator.AddShaderResourceViewParameter(RAY_TRACING_SYSTEM_VERTEXBUFFER_REGISTER, UE_HLSL_SPACE_RAY_TRACING_SYSTEM);
 
-			uint32 NumConstants = sizeof(FHitGroupSystemRootConstants) / sizeof(uint32);
-			Creator.AddConstantsParameter(NumConstants, RAY_TRACING_SYSTEM_ROOTCONSTANT_REGISTER, UE_HLSL_SPACE_RAY_TRACING_SYSTEM);
+				uint32 NumConstants = sizeof(FHitGroupSystemRootConstants) / sizeof(uint32);
+				Creator.AddConstantsParameter(NumConstants, RAY_TRACING_SYSTEM_ROOTCONSTANT_REGISTER, UE_HLSL_SPACE_RAY_TRACING_SYSTEM);
+			}
 		}
 
 		if (Creator.HasFlags(EShaderBindingLayoutFlags::BindlessResources))
diff --git a/Engine/Source/Runtime/Renderer/Private/RayTracing/RayTracingDynamicGeometry.cpp b/Engine/Source/Runtime/Renderer/Private/RayTracing/RayTracingDynamicGeometry.cpp
--- a/Engine/Source/Runtime/Renderer/Private/RayTracing/RayTracingDynamicGeometry.cpp
+++ b/Engine/Source/Runtime/Renderer/Private/RayTracing/RayTracingDynamicGeometry.cpp
@@ -236,9 +236,13 @@ void FRayTracingDynamicGeometryCollection::AddDynamicMeshBatchForGeometryUpdate(
 		VertexBufferOffset = VertexPositionBuffer->UsedSize;
 		VertexPositionBuffer->UsedSize += UpdateParams.VertexBufferSize;
 
+		// Make sure vertex buffer offset is aligned to 16 (required for Raw SRV views)
+		VertexPositionBuffer->UsedSize = Align(VertexPositionBuffer->UsedSize, 16);
+
 		bUseSharedVertexBuffer = true;
 		RWBuffer = &VertexPositionBuffer->RWBuffer;
 	}
+	check(IsAligned(VertexBufferOffset, 16));
 
 	FRayTracingDynamicGeometryBuildParams GeometryBuildParams;
 	GeometryBuildParams.DispatchCommands.Reserve(UpdateParams.MeshBatches.Num());
