Add RHI breadcrumb allocator array to GPU profiler event stream
 - Removes the burden of managing the lifetime of the breadcrumbs emitted in the stream from the platform RHIs.
 - Will allow the event stream to be paused/resumed by the profilers without hitting memory lifetime issues. This will be necessary when we start processing fence signal/wait events.

Add support for signal / wait events in D3D12 RHI
 - Refactored the initialization of the queue timing structs to avoid cases were signal / wait events are generated before a timing struct exists to add them to.
 - Removed FlushTiming and merged a simplified version into RHIEndFrame.
 - Removed checks that skip copy queue profiler initialization. We now emit signal / wait / frame boundary events for the copy queue.

diff --git a/Engine/Source/Runtime/D3D12RHI/Private/D3D12CommandList.h b/Engine/Source/Runtime/D3D12RHI/Private/D3D12CommandList.h
--- a/Engine/Source/Runtime/D3D12RHI/Private/D3D12CommandList.h
+++ b/Engine/Source/Runtime/D3D12RHI/Private/D3D12CommandList.h
@@ -123,6 +123,7 @@ public:
 
 		if constexpr (std::is_same_v<UE::RHI::GPUProfiler::FEvent::FBeginWork, TEventType>)
 		{
+			// Store BeginEvents in a separate array as the CPUTimestamp field needs updating at submit time.
 			State.BeginEvents.Add(&Data);
 		}
 
diff --git a/Engine/Source/Runtime/D3D12RHI/Private/D3D12Device.h b/Engine/Source/Runtime/D3D12RHI/Private/D3D12Device.h
--- a/Engine/Source/Runtime/D3D12RHI/Private/D3D12Device.h
+++ b/Engine/Source/Runtime/D3D12RHI/Private/D3D12Device.h
@@ -73,10 +73,6 @@ public:
 	uint64 GPUFrequency = 0, GPUTimestamp = 0;
 	uint64 CPUFrequency = 0, CPUTimestamp = 0;
 
-	#if WITH_RHI_BREADCRUMBS
-	TArray<TSharedPtr<FRHIBreadcrumbAllocatorArray>> BreadcrumbAllocators {};
-	#endif
-
 	UE::RHI::GPUProfiler::FEventStream EventStream;
 
 #else
diff --git a/Engine/Source/Runtime/D3D12RHI/Private/D3D12RHI.cpp b/Engine/Source/Runtime/D3D12RHI/Private/D3D12RHI.cpp
--- a/Engine/Source/Runtime/D3D12RHI/Private/D3D12RHI.cpp
+++ b/Engine/Source/Runtime/D3D12RHI/Private/D3D12RHI.cpp
@@ -413,49 +413,6 @@ void FD3D12DynamicRHI::EnqueueEndOfPipeTask(TUniqueFunction<void()> TaskFunc, TU
 	);
 }
 
-void FD3D12DynamicRHI::FlushTiming(bool bCreateNew, const FRHIEndFrameArgs& Args)
-{
-	bool bHasTiming = !CurrentTimingPerQueue.IsEmpty();
-	auto Lambda = [this, OldTiming = MoveTemp(CurrentTimingPerQueue)]()
-	{
-		ProcessTimestamps(OldTiming);
-	};
-
-	EnqueueEndOfPipeTask(MoveTemp(Lambda), [&](FD3D12Payload& Payload)
-	{
-		if (bCreateNew)
-		{
-			TUniquePtr<FD3D12Timing> NewTiming = MakeUnique<FD3D12Timing>(Payload.Queue);
-			Payload.Timing = NewTiming.Get();
-			CurrentTimingPerQueue.Emplace(MoveTemp(NewTiming));
-		}
-		else
-		{
-			Payload.Timing = nullptr;
-		}
-
-#if RHI_NEW_GPU_PROFILER
-		if (bHasTiming && Payload.Queue.QueueType != ED3D12QueueType::Copy)
-		{
-			ERHIPipeline Pipeline;
-			switch (Payload.Queue.QueueType)
-			{
-			default: checkNoEntry(); [[fallthrough]];
-			case ED3D12QueueType::Direct: Pipeline = ERHIPipeline::Graphics; break;
-			case ED3D12QueueType::Async:  Pipeline = ERHIPipeline::AsyncCompute; break;
-			}
-
-			Payload.EventStream.Emplace<UE::RHI::GPUProfiler::FEvent::FFrameBoundary>(
-				Args.FrameNumber
-		#if WITH_RHI_BREADCRUMBS
-				, Args.GPUBreadcrumbs[Pipeline]
-		#endif
-			);
-		}
-#endif // RHI_NEW_GPU_PROFILER
-	});
-}
-
 void FD3D12DynamicRHI::RHIProcessDeleteQueue()
 {
 	ProcessDeferredDeletionQueue_Platform();
@@ -613,7 +570,40 @@ void FD3D12DynamicRHI::RHIEndFrame(const FRHIEndFrameArgs& Args)
 	UpdateMemoryStats();
 
 	// Close the previous frame's timing and start a new one
-	FlushTiming(true, Args);
+	auto Lambda = [this, OldTiming = MoveTemp(CurrentTimingPerQueue)]()
+	{
+		ProcessTimestamps(OldTiming);
+	};
+
+	EnqueueEndOfPipeTask(MoveTemp(Lambda), [&](FD3D12Payload& Payload)
+	{
+		// Modify the payloads the EOP task will submit to include
+		// a new timing struct and a frame boundary event.
+
+		Payload.Timing = CurrentTimingPerQueue.CreateNew(Payload.Queue);
+
+	#if RHI_NEW_GPU_PROFILER
+		ERHIPipeline Pipeline;
+		switch (Payload.Queue.QueueType)
+		{
+		default: checkNoEntry(); [[fallthrough]];
+		case ED3D12QueueType::Direct: Pipeline = ERHIPipeline::Graphics; break;
+		case ED3D12QueueType::Async:  Pipeline = ERHIPipeline::AsyncCompute; break;
+
+		case ED3D12QueueType::Copy:
+			// There is currently no high level RHI copy queue support
+			Pipeline = ERHIPipeline::None;
+			break;
+		}
+
+		Payload.EventStream.Emplace<UE::RHI::GPUProfiler::FEvent::FFrameBoundary>(
+			Args.FrameNumber
+		#if WITH_RHI_BREADCRUMBS
+			, (Pipeline != ERHIPipeline::None) ? Args.GPUBreadcrumbs[Pipeline] : nullptr
+		#endif
+		);
+	#endif
+	});
 
 	// Pump the interrupt queue to gather completed events
 	// (required if we're not using an interrupt thread).
diff --git a/Engine/Source/Runtime/D3D12RHI/Private/D3D12RHIPrivate.h b/Engine/Source/Runtime/D3D12RHI/Private/D3D12RHIPrivate.h
--- a/Engine/Source/Runtime/D3D12RHI/Private/D3D12RHIPrivate.h
+++ b/Engine/Source/Runtime/D3D12RHI/Private/D3D12RHIPrivate.h
@@ -235,10 +235,15 @@ public:
 	// Called when the GPU has crashed. This function will not return.
 	void ProcessInterruptQueueOnGPUCrash();
 
-	using FD3D12TimingArray = TArray<TUniquePtr<FD3D12Timing>, TInlineAllocator<GD3D12MaxNumQueues>>;
+	struct FD3D12TimingArray : public TArray<TUniquePtr<FD3D12Timing>, TInlineAllocator<GD3D12MaxNumQueues>>
+	{
+		FD3D12Timing* CreateNew(FD3D12Queue& Queue)
+		{
+			return Emplace_GetRef(MakeUnique<FD3D12Timing>(Queue)).Get();
+		}
+	};
 
 	FD3D12TimingArray CurrentTimingPerQueue;
-	void FlushTiming(bool bCreateNew, const FRHIEndFrameArgs& Args);
 	void ProcessTimestamps(FD3D12TimingArray const& TimingPerQueue);
 
 	void InitializeSubmissionPipe();
diff --git a/Engine/Source/Runtime/D3D12RHI/Private/D3D12Submission.cpp b/Engine/Source/Runtime/D3D12RHI/Private/D3D12Submission.cpp
--- a/Engine/Source/Runtime/D3D12RHI/Private/D3D12Submission.cpp
+++ b/Engine/Source/Runtime/D3D12RHI/Private/D3D12Submission.cpp
@@ -119,25 +119,6 @@ private:
 
 void FD3D12DynamicRHI::InitializeSubmissionPipe()
 {
-#if RHI_NEW_GPU_PROFILER
-	{
-		TArray<UE::RHI::GPUProfiler::FQueue> Queues;
-		ForEachQueue([&](FD3D12Queue& Queue)
-		{
-			if (Queue.QueueType == ED3D12QueueType::Copy)
-			{
-				// @todo dev-pr : Skip copy queues for now. We don't have copy queue support in RHI contexts,
-				// so it's not possible to insert frame boundary events in a way that is pipelined correctly.
-				return;
-			}
-
-			Queues.Add(Queue.GetProfilerQueue());
-		});
-
-		UE::RHI::GPUProfiler::InitializeQueues(Queues);
-	}
-#endif
-
 	if (FPlatformProcess::SupportsMultithreading())
 	{
 #if D3D12_USE_INTERRUPT_THREAD
@@ -162,33 +143,32 @@ void FD3D12DynamicRHI::InitializeSubmissionPipe()
 #endif
 	}
 
-#if WITH_RHI_BREADCRUMBS
-	TRHIPipelineArray<FRHIBreadcrumbNode*> GPUBreadcrumbs { InPlace, nullptr };
-	FRHIEndFrameArgs Args
+	// Initialize the timing structs in each queue, and the engine GPU profilers
 	{
-		.GPUBreadcrumbs = GPUBreadcrumbs
-	};
-#else
-	FRHIEndFrameArgs Args;
-#endif
-	
-	FlushTiming(true, Args);
+		TArray<FD3D12Payload*> Payloads;
+	#if RHI_NEW_GPU_PROFILER
+		TArray<UE::RHI::GPUProfiler::FQueue> ProfilerQueues;
+	#endif
+
+		ForEachQueue([&](FD3D12Queue& Queue)
+		{
+			FD3D12Payload* Payload = Payloads.Emplace_GetRef(new FD3D12Payload(Queue));
+			Payload->Timing = CurrentTimingPerQueue.CreateNew(Queue);
+
+		#if RHI_NEW_GPU_PROFILER
+			ProfilerQueues.Add(Queue.GetProfilerQueue());
+		#endif
+		});
+
+	#if RHI_NEW_GPU_PROFILER
+		UE::RHI::GPUProfiler::InitializeQueues(ProfilerQueues);
+	#endif
+		SubmitPayloads(MoveTemp(Payloads));
+	}
 }
 
 void FD3D12DynamicRHI::ShutdownSubmissionPipe()
 {
-#if WITH_RHI_BREADCRUMBS
-	TRHIPipelineArray<FRHIBreadcrumbNode*> GPUBreadcrumbs { InPlace, nullptr };
-	FRHIEndFrameArgs Args
-	{
-		.GPUBreadcrumbs = GPUBreadcrumbs
-	};
-#else
-	FRHIEndFrameArgs Args;
-#endif
-
-	FlushTiming(false, Args);
-
 	delete SubmissionThread;
 	SubmissionThread = nullptr;
 
@@ -271,7 +251,8 @@ void FD3D12DynamicRHI::SubmitCommands(TConstArrayView<FD3D12FinalizedCommands*>
 			Payload->BreadcrumbRange = Payloads->BreadcrumbRange;
 			if (BreadcrumbAllocators.IsValid())
 			{
-				Payload->BatchedObjects.BreadcrumbAllocators.Add(BreadcrumbAllocators);
+				check(!Payload->BreadcrumbAllocators.IsValid());
+				Payload->BreadcrumbAllocators = BreadcrumbAllocators;
 			}
 		}
 	#endif
@@ -735,13 +716,6 @@ uint64 FD3D12Queue::FinalizePayload(bool bRequiresSignal, FPayloadArray& Payload
 
 	BarrierTimestamps.CloseAndReset(PayloadToSubmit->BatchedObjects.QueryRanges);
 
-#if WITH_RHI_BREADCRUMBS && RHI_NEW_GPU_PROFILER
-	if (PayloadToSubmit->BatchedObjects.BreadcrumbAllocators.Num())
-	{
-		BatchedObjects.BreadcrumbAllocators.Append(MoveTemp(PayloadToSubmit->BatchedObjects.BreadcrumbAllocators));
-	}
-#endif
-
 	// Gather query ranges from this payload, grouping by heap pointer
 	if (BatchedObjects.QueryRanges.Num())
 	{
@@ -882,6 +856,15 @@ void FD3D12DynamicRHI::FlushBatchedPayloads(FD3D12Queue::FPayloadArray& Payloads
 		// Wait for queue fences
 		for (auto& [LocalFence, Value] : Payload->QueueFencesToWait)
 		{
+		#if RHI_NEW_GPU_PROFILER
+			// Use the raw ID3D12Fence pointer as the fence's unique ID, since we never release these.
+			Payload->EventStream.Emplace<UE::RHI::GPUProfiler::FEvent::FWaitFence>(
+				  FPlatformTime::Cycles64()
+				, reinterpret_cast<uint64>(LocalFence.D3DFence.GetReference())
+				, Value
+			);
+		#endif
+
 			VERIFYD3D12RESULT(Queue.D3DCommandQueue->Wait(LocalFence.D3DFence, Value));
 		}
 
@@ -1024,6 +1007,15 @@ void FD3D12DynamicRHI::FlushBatchedPayloads(FD3D12Queue::FPayloadArray& Payloads
 		{
 			check(Queue.Fence.LastSignaledValue < Payload->CompletionFenceValue);
 
+		#if RHI_NEW_GPU_PROFILER
+			// Use the raw ID3D12Fence pointer as the fence's unique ID, since we never release these.
+			Payload->EventStream.Emplace<UE::RHI::GPUProfiler::FEvent::FSignalFence>(
+				  FPlatformTime::Cycles64()
+				, reinterpret_cast<uint64>(Queue.Fence.D3DFence.GetReference())
+				, Payload->CompletionFenceValue
+			);
+		#endif
+
 			VERIFYD3D12RESULT(Queue.D3DCommandQueue->Signal(Queue.Fence.D3DFence, Payload->CompletionFenceValue));
 			Queue.Fence.LastSignaledValue.store(Payload->CompletionFenceValue, std::memory_order_release);
 		}
@@ -1373,7 +1365,7 @@ FD3D12DynamicRHI::FProcessResult FD3D12DynamicRHI::ProcessInterruptQueue()
 							}
 							break;
 
-#if RHI_NEW_GPU_PROFILER
+					#if RHI_NEW_GPU_PROFILER
 						case ED3D12QueryType::ProfilerTimestampTOP:
 						case ED3D12QueryType::ProfilerTimestampBOP:
 							{
@@ -1386,7 +1378,7 @@ FD3D12DynamicRHI::FProcessResult FD3D12DynamicRHI::ProcessInterruptQueue()
 								Target = CPUDelta + CurrentQueue.Timing->CPUTimestamp;
 							}
 							break;
-#else
+					#else
 						case ED3D12QueryType::CommandListBegin:
 						case ED3D12QueryType::CommandListEnd:
 						case ED3D12QueryType::IdleBegin:
@@ -1394,26 +1386,19 @@ FD3D12DynamicRHI::FProcessResult FD3D12DynamicRHI::ProcessInterruptQueue()
 							check(CurrentQueue.Timing);
 							CurrentQueue.Timing->Timestamps.Add(Query.GetResult<uint64>());
 							break;
-
-#endif
+					#endif
 						}
 					}
 				}
 			}
 
-#if RHI_NEW_GPU_PROFILER && WITH_RHI_BREADCRUMBS
-			if (Payload->BatchedObjects.BreadcrumbAllocators.Num())
-			{
-				check(CurrentQueue.Timing);
-				CurrentQueue.Timing->BreadcrumbAllocators.Append(MoveTemp(Payload->BatchedObjects.BreadcrumbAllocators));
-			}
-
+		#if RHI_NEW_GPU_PROFILER
 			if (!Payload->EventStream.IsEmpty())
 			{
 				check(CurrentQueue.Timing);
 				CurrentQueue.Timing->EventStream.Append(MoveTemp(Payload->EventStream));
 			}
-#endif
+		#endif
 
 			if (Payload->Timing.IsSet())
 			{
@@ -1479,7 +1464,7 @@ void FD3D12DynamicRHI::ProcessTimestamps(FD3D12TimingArray const& TimingPerQueue
 
 	for (auto const& Timing : TimingPerQueue)
 	{
-		UE::RHI::GPUProfiler::ProcessEvents(Timing->Queue.GetProfilerQueue(), Timing->EventStream);
+		UE::RHI::GPUProfiler::ProcessEvents(Timing->Queue.GetProfilerQueue(), MoveTemp(Timing->EventStream));
 	}
 
 #else
diff --git a/Engine/Source/Runtime/D3D12RHI/Private/D3D12Submission.h b/Engine/Source/Runtime/D3D12RHI/Private/D3D12Submission.h
--- a/Engine/Source/Runtime/D3D12RHI/Private/D3D12Submission.h
+++ b/Engine/Source/Runtime/D3D12RHI/Private/D3D12Submission.h
@@ -148,10 +148,6 @@ struct FD3D12BatchedPayloadObjects
 	TArray<FD3D12QueryLocation> PipelineStatsQueries;
 	TMap<TRefCountPtr<FD3D12QueryHeap>, TArray<FD3D12QueryRange>> QueryRanges;
 
-#if WITH_RHI_BREADCRUMBS && RHI_NEW_GPU_PROFILER
-	TArray<TSharedPtr<FRHIBreadcrumbAllocatorArray>, TInlineAllocator<1>> BreadcrumbAllocators {};
-#endif
-
 	bool IsEmpty() const
 	{
 		return
@@ -159,11 +155,7 @@ struct FD3D12BatchedPayloadObjects
 			&& OcclusionQueries    .Num() == 0
 			&& PipelineStatsQueries.Num() == 0
 			&& QueryRanges         .Num() == 0
-#if WITH_RHI_BREADCRUMBS && RHI_NEW_GPU_PROFILER
-			&& BreadcrumbAllocators.Num() == 0
-#endif
 		;
-
 	}
 };
 
@@ -244,7 +236,9 @@ struct FD3D12PayloadBase : public FD3D12PayloadBaseFixLayout
 
 #if WITH_RHI_BREADCRUMBS
 	FRHIBreadcrumbRange BreadcrumbRange {};
+	TSharedPtr<FRHIBreadcrumbAllocatorArray> BreadcrumbAllocators {};
 #endif
+
 #if RHI_NEW_GPU_PROFILER
 	UE::RHI::GPUProfiler::FEventStream EventStream;
 #endif
diff --git a/Engine/Source/Runtime/RHI/Private/GPUProfiler.cpp b/Engine/Source/Runtime/RHI/Private/GPUProfiler.cpp
--- a/Engine/Source/Runtime/RHI/Private/GPUProfiler.cpp
+++ b/Engine/Source/Runtime/RHI/Private/GPUProfiler.cpp
@@ -724,8 +724,10 @@ namespace UE::RHI::GPUProfiler
 		GetSinks().RemoveSingle(this);
 	}
 
-	void ProcessEvents(FQueue Queue, FEventStream const& EventStream)
+	void ProcessEvents(FQueue Queue, FEventStream EventStream)
 	{
+		TRACE_CPUPROFILER_EVENT_SCOPE(UE::RHI::GPUProfiler::ProcessEvents);
+
 		if (!EventStream.IsEmpty())
 		{
 			for (FEventSink* Sink : GetSinks())
@@ -1006,6 +1008,20 @@ namespace UE::RHI::GPUProfiler
 					break;
 			#endif // WITH_RHI_BREADCRUMBS
 
+				case FEvent::EType::SignalFence:
+					{
+						check(!QueueState.bBusy);
+						FEvent::FSignalFence const& SignalFence = Event->Value.Get<FEvent::FSignalFence>();
+					}
+					break;
+
+				case FEvent::EType::WaitFence:
+					{
+						check(!QueueState.bBusy);
+						FEvent::FWaitFence const& WaitFence = Event->Value.Get<FEvent::FWaitFence>();
+					}
+					break;
+
 				case FEvent::EType::FrameBoundary:
 					{
 						check(!QueueState.bBusy);
diff --git a/Engine/Source/Runtime/RHI/Public/GPUProfiler.h b/Engine/Source/Runtime/RHI/Public/GPUProfiler.h
--- a/Engine/Source/Runtime/RHI/Public/GPUProfiler.h
+++ b/Engine/Source/Runtime/RHI/Public/GPUProfiler.h
@@ -114,7 +114,7 @@ namespace UE::RHI::GPUProfiler
 	#if WITH_RHI_BREADCRUMBS
 		struct FBeginBreadcrumb
 		{
-			FRHIBreadcrumbNode* Breadcrumb;
+			FRHIBreadcrumbNode* const Breadcrumb;
 			uint64 GPUTimestampTOP = 0;
 
 			FBeginBreadcrumb(FRHIBreadcrumbNode* Breadcrumb)
@@ -124,7 +124,7 @@ namespace UE::RHI::GPUProfiler
 
 		struct FEndBreadcrumb
 		{
-			FRHIBreadcrumbNode* Breadcrumb;
+			FRHIBreadcrumbNode* const Breadcrumb;
 			uint64 GPUTimestampBOP = 0;
 
 			FEndBreadcrumb(FRHIBreadcrumbNode* Breadcrumb)
@@ -164,16 +164,47 @@ namespace UE::RHI::GPUProfiler
 			}
 		};
 
+		// Can only be inserted when the GPU is marked "idle", i.e. after an FEndWork event.
 		struct FSignalFence
 		{
+			//
+			// Timestamp when the fence signal was enqueued to the GPU/driver.
+			// 
+			// The signal on the GPU doesn't happen until after the previous FEndWork
+			// event's BOP timestamp, or this CPU timestamp, whichever is later.
+			//
 			uint64 CPUTimestamp;
+
+			// Unique ID of the fence signaled.
 			uint64 ID;
+
+			// The fence value signaled.
+			uint64 Value;
+
+			FSignalFence(uint64 CPUTimestamp, uint64 ID, uint64 Value)
+				: CPUTimestamp(CPUTimestamp)
+				, ID(ID)
+				, Value(Value)
+			{}
 		};
 
+		// Can only be inserted when the GPU is marked "idle", i.e. after an FEndWork event.
 		struct FWaitFence
 		{
+			// Timestamp when the fence wait was enqueued to the GPU/driver.
 			uint64 CPUTimestamp;
+
+			// Unique ID of the fence awaited.
 			uint64 ID;
+
+			// The fence value awaited.
+			uint64 Value;
+
+			FWaitFence(uint64 CPUTimestamp, uint64 ID, uint64 Value)
+				: CPUTimestamp(CPUTimestamp)
+				, ID(ID)
+				, Value(Value)
+			{}
 		};
 
 		struct FFlip
@@ -242,6 +273,10 @@ namespace UE::RHI::GPUProfiler
 			{
 				FChunk* Next = nullptr;
 				uint32 Num = 0;
+
+			#if WITH_RHI_BREADCRUMBS
+				FRHIBreadcrumbAllocatorArray BreadcrumbAllocators;
+			#endif
 			} Header;
 
 			static constexpr uint32 ChunkSizeInBytes = 16 * 1024;
@@ -275,7 +310,7 @@ namespace UE::RHI::GPUProfiler
 			}
 		};
 
-		static_assert(sizeof(FChunk) == FChunk::ChunkSizeInBytes, "Incorrect FChunk size.");
+		static_assert(sizeof(FChunk) <= FChunk::ChunkSizeInBytes, "Incorrect FChunk size.");
 
 		FChunk* First = nullptr;
 		FChunk* Current = nullptr;
@@ -305,7 +340,7 @@ namespace UE::RHI::GPUProfiler
 		template <typename TEventType, typename... TArgs>
 		TEventType& Emplace(TArgs&&... Args)
 		{
-			static_assert(TIsTrivial<TEventType>::Value, "Destructors are not called on GPU profiler events, so the types must be trivial.");
+			static_assert(std::is_trivially_destructible_v<TEventType>, "Destructors are not called on GPU profiler events, so the types must be trivially destructible.");
 
 			if (!Current)
 			{
@@ -326,7 +361,21 @@ namespace UE::RHI::GPUProfiler
 			FEvent* Event = Current->GetElement(Current->Header.Num++);
 			new (Event) FEvent(TEventType(Forward<TArgs>(Args)...));
 
-			return Event->Value.Get<TEventType>();
+			TEventType& Data = Event->Value.Get<TEventType>();
+
+		#if WITH_RHI_BREADCRUMBS
+			if constexpr (
+				std::is_same_v<UE::RHI::GPUProfiler::FEvent::FBeginBreadcrumb, TEventType> ||
+				std::is_same_v<UE::RHI::GPUProfiler::FEvent::FEndBreadcrumb  , TEventType>
+				)
+			{
+				// Attach the breadcrumb allocator for begin/end breadcrumb events.
+				// This keeps the breadcrumbs alive until the events have been consumed by the profilers.
+				Current->Header.BreadcrumbAllocators.AddUnique(Data.Breadcrumb->Allocator);
+			}
+		#endif
+
+			return Data;
 		}
 
 		bool IsEmpty() const
@@ -412,7 +461,7 @@ namespace UE::RHI::GPUProfiler
 		virtual void InitializeQueues(TConstArrayView<FQueue> Queues) = 0;
 	};
 
-	RHI_API void ProcessEvents(FQueue Queue, FEventStream const& EventStream);
+	RHI_API void ProcessEvents(FQueue Queue, FEventStream EventStream);
 	RHI_API void InitializeQueues(TConstArrayView<FQueue> Queues);
 }
 
diff --git a/Engine/Source/Runtime/Windows/D3D11RHI/Private/D3D11RHI.cpp b/Engine/Source/Runtime/Windows/D3D11RHI/Private/D3D11RHI.cpp
--- a/Engine/Source/Runtime/Windows/D3D11RHI/Private/D3D11RHI.cpp
+++ b/Engine/Source/Runtime/Windows/D3D11RHI/Private/D3D11RHI.cpp
@@ -79,7 +79,7 @@ void FD3D11DynamicRHI::RHIEndFrame(const FRHIEndFrameArgs& Args)
 			PollQueryResults();
 
 			// Previous frame has completed and the data is available. Publish the profiler events.
-			UE::RHI::GPUProfiler::ProcessEvents(UE::RHI::GPUProfiler::FQueue(UE::RHI::GPUProfiler::FQueue::EType::Graphics, 0, 0), PreviousFrame->EventStream);
+			UE::RHI::GPUProfiler::ProcessEvents(UE::RHI::GPUProfiler::FQueue(UE::RHI::GPUProfiler::FQueue::EType::Graphics, 0, 0), MoveTemp(PreviousFrame->EventStream));
 
 			Profiler.EventPool.Push(MoveTemp(PreviousFrame->CompletionQuery));
 			Profiler.Pending.Pop();
