Fixing D3D12 Bindless when static root signatures are used.
This code was lost in a recent refactor.

diff --git a/Engine/Source/Runtime/D3D12RHI/Private/D3D12Adapter.cpp b/Engine/Source/Runtime/D3D12RHI/Private/D3D12Adapter.cpp
--- a/Engine/Source/Runtime/D3D12RHI/Private/D3D12Adapter.cpp
+++ b/Engine/Source/Runtime/D3D12RHI/Private/D3D12Adapter.cpp
@@ -1185,7 +1185,10 @@ void FD3D12Adapter::InitializeDevices()
 		// (because we will create placed buffers for texture allocation to retrieve the GPU virtual addresses)
 		const bool bTraceMemAlloc = UE_TRACE_CHANNELEXPR_IS_ENABLED(MemAllocChannel);
 		bTrackAllAllocation = (GD3D12TrackAllAlocations || UE::RHI::UseGPUCrashDebugging() || bTraceMemAlloc) && (GetResourceHeapTier() == D3D12_RESOURCE_HEAP_TIER_2);
-#endif 
+#endif
+
+		ERHIBindlessConfiguration BindlessResourcesConfig = ERHIBindlessConfiguration::Disabled;
+		ERHIBindlessConfiguration BindlessSamplersConfig = ERHIBindlessConfiguration::Disabled;
 
 #if PLATFORM_SUPPORTS_BINDLESS_RENDERING
 		if (Desc.MaxSupportedFeatureLevel >= D3D_FEATURE_LEVEL_12_0 && Desc.MaxSupportedShaderModel >= D3D_SHADER_MODEL_6_6 && Desc.ResourceBindingTier >= D3D12_RESOURCE_BINDING_TIER_3)
@@ -1193,7 +1196,10 @@ void FD3D12Adapter::InitializeDevices()
 			// Needs to happen before device creation below
 			BindlessManager.Init(this);
 
-			GRHIGlobals.ShaderBundles.RequiresSharedBindlessParameters = (BindlessManager.GetResourcesConfiguration() == ERHIBindlessConfiguration::AllShaders || BindlessManager.GetSamplersConfiguration() == ERHIBindlessConfiguration::AllShaders);
+			BindlessResourcesConfig = BindlessManager.GetResourcesConfiguration();
+			BindlessSamplersConfig = BindlessManager.GetSamplersConfiguration();
+
+			GRHIGlobals.ShaderBundles.RequiresSharedBindlessParameters = (BindlessResourcesConfig == ERHIBindlessConfiguration::AllShaders || BindlessSamplersConfig == ERHIBindlessConfiguration::AllShaders);
 		}
 #endif
 
@@ -1237,9 +1243,6 @@ void FD3D12Adapter::InitializeDevices()
 		PipelineStateCache.Init(GraphicsCacheFile, ComputeCacheFile, DriverBlobFilename);
 		PipelineStateCache.RebuildFromDiskCache();
 
-		ERHIBindlessConfiguration BindlessResourcesConfig = ERHIBindlessConfiguration::Disabled;
-		ERHIBindlessConfiguration BindlessSamplersConfig = ERHIBindlessConfiguration::Disabled;
-
 #if USE_STATIC_ROOT_SIGNATURE
 		EShaderBindingLayoutFlags GraphicsFlags{};
 #if PLATFORM_SUPPORTS_MESH_SHADERS
