Don't collect or report RHI breadcrumbs if GPU marker writing is disabled.

diff --git a/Engine/Source/Runtime/D3D12RHI/Private/D3D12Util.cpp b/Engine/Source/Runtime/D3D12RHI/Private/D3D12Util.cpp
--- a/Engine/Source/Runtime/D3D12RHI/Private/D3D12Util.cpp
+++ b/Engine/Source/Runtime/D3D12RHI/Private/D3D12Util.cpp
@@ -832,42 +832,50 @@ void FD3D12DynamicRHI::OutputGPUCrashReport(FTextBuilder& ErrorMessage)
 				ShaderDiagnostics += Queue.DiagnosticBuffer->GetShaderDiagnosticMessages(DeviceIndex, QueueIndex, QueueName);
 
 #if WITH_RHI_BREADCRUMBS
-				ERHIPipeline Pipeline;
-				switch (Queue.QueueType)
+				// Don't collect breadcrumb ranges if we don't have breadcrumbs enabled. The breadcrumbs
+				// will not have meaningful GPU state information because GPU markers are not written.
+				if (UE::RHI::UseGPUCrashBreadcrumbs())
 				{
-				default: continue; // Skip pipelines that the RHI doesn't handle
-				case ED3D12QueueType::Direct: Pipeline = ERHIPipeline::Graphics; break;
-				case ED3D12QueueType::Async : Pipeline = ERHIPipeline::AsyncCompute; break;
-				}
+					ERHIPipeline Pipeline;
+					switch (Queue.QueueType)
+					{
+						default: continue; // Skip pipelines that the RHI doesn't handle
+						case ED3D12QueueType::Direct: Pipeline = ERHIPipeline::Graphics; break;
+						case ED3D12QueueType::Async: Pipeline = ERHIPipeline::AsyncCompute; break;
+					}
 
-				TArray<FRHIBreadcrumbRange>& Ranges = QueueRanges.Add({ DeviceIndex, Pipeline });
+					TArray<FRHIBreadcrumbRange>& Ranges = QueueRanges.Add({ DeviceIndex, Pipeline });
 
-				// Pull all incomplete payloads from the pending interrupt queue.
-				TArray<FD3D12Payload*> Payloads;
-				{
-					FD3D12Payload* Payload;
-					while (Queue.PendingInterrupt.Dequeue(Payload))
-						Payloads.Add(Payload);
-				}
+					// Pull all incomplete payloads from the pending interrupt queue.
+					TArray<FD3D12Payload*> Payloads;
+					{
+						FD3D12Payload* Payload;
+						while (Queue.PendingInterrupt.Dequeue(Payload))
+							Payloads.Add(Payload);
+					}
 
-				// Extract the breadcrumb ranges for these payloads.
-				for (FD3D12Payload* Payload : Payloads)
-				{
-					if (Payload->BreadcrumbRange)
+					// Extract the breadcrumb ranges for these payloads.
+					for (FD3D12Payload* Payload : Payloads)
 					{
-						Ranges.AddUnique(Payload->BreadcrumbRange);
+						if (Payload->BreadcrumbRange)
+						{
+							Ranges.AddUnique(Payload->BreadcrumbRange);
+						}
 					}
-				}
 
-				BreadcrumbState.Devices[DeviceIndex].Pipelines[Pipeline].MarkerOut = Queue.DiagnosticBuffer->ReadMarkerOut();
-				BreadcrumbState.Devices[DeviceIndex].Pipelines[Pipeline].MarkerIn  = Queue.DiagnosticBuffer->ReadMarkerIn();
+					BreadcrumbState.Devices[DeviceIndex].Pipelines[Pipeline].MarkerOut = Queue.DiagnosticBuffer->ReadMarkerOut();
+					BreadcrumbState.Devices[DeviceIndex].Pipelines[Pipeline].MarkerIn = Queue.DiagnosticBuffer->ReadMarkerIn();
+				}
 #endif
 			}
 		}
 
 #if WITH_RHI_BREADCRUMBS
 		// Traverse the breadcrumb tree and log active GPU work
-		BreadcrumbState.DumpActiveBreadcrumbs(QueueRanges);
+		if (!QueueRanges.IsEmpty())
+		{
+			BreadcrumbState.DumpActiveBreadcrumbs(QueueRanges);
+		}
 #endif // WITH_RHI_BREADCRUMBS
 
 		if (!ShaderDiagnostics.IsEmpty())
