Fixing compilation errors in Debug builds.

diff --git a/Engine/Source/Runtime/D3D12RHI/Private/D3D12RHICommon.h b/Engine/Source/Runtime/D3D12RHI/Private/D3D12RHICommon.h
--- a/Engine/Source/Runtime/D3D12RHI/Private/D3D12RHICommon.h
+++ b/Engine/Source/Runtime/D3D12RHI/Private/D3D12RHICommon.h
@@ -112,8 +112,8 @@ public:
 	FORCEINLINE FRHIGPUMask GetGPUMask() const { return GPUMask; }
 	FORCEINLINE FRHIGPUMask GetVisibilityMask() const { return VisibilityMask; }
 #else
-	constexpr FRHIGPUMask GetGPUMask() const { return FRHIGPUMask::GPU0(); }
-	constexpr FRHIGPUMask GetVisibilityMask() const { return FRHIGPUMask::GPU0(); }
+	SGPU_CONSTEXPR FRHIGPUMask GetGPUMask() const { return FRHIGPUMask::GPU0(); }
+	SGPU_CONSTEXPR FRHIGPUMask GetVisibilityMask() const { return FRHIGPUMask::GPU0(); }
 #endif
 
 protected:
@@ -140,7 +140,7 @@ public:
 		return GPUIndex;
 	}
 #else
-	constexpr uint32 GetGPUIndex() const { return 0; }
+	SGPU_CONSTEXPR uint32 GetGPUIndex() const { return 0; }
 #endif
 
 private:
