Fix warning C6237: (<zero> && <expression>) is always zero.  <expression> is never evaluated and might have side effects.
#rnx

diff --git a/Engine/Source/Runtime/D3D12RHI/Private/D3D12Adapter.cpp b/Engine/Source/Runtime/D3D12RHI/Private/D3D12Adapter.cpp
--- a/Engine/Source/Runtime/D3D12RHI/Private/D3D12Adapter.cpp
+++ b/Engine/Source/Runtime/D3D12RHI/Private/D3D12Adapter.cpp
@@ -348,9 +348,11 @@ void FD3D12Adapter::CreateRootDevice(bool bWithDebug)
 	// -d3ddebug is always allowed on Windows, but only allowed in non-shipping builds on other platforms.
 	// -gpuvalidation is only supported on Windows.
 #if PLATFORM_WINDOWS || !UE_BUILD_SHIPPING
-	bool bWithGPUValidation = PLATFORM_WINDOWS && (FParse::Param(FCommandLine::Get(), TEXT("d3d12gpuvalidation")) || FParse::Param(FCommandLine::Get(), TEXT("gpuvalidation")));
+#if PLATFORM_WINDOWS
+	bool bWithGPUValidation =  (FParse::Param(FCommandLine::Get(), TEXT("d3d12gpuvalidation")) || FParse::Param(FCommandLine::Get(), TEXT("gpuvalidation")));
 	// If GPU validation is requested, automatically enable the debug layer.
 	bWithDebug |= bWithGPUValidation;
+#endif
 	if (bWithDebug)
 	{
 		TRefCountPtr<ID3D12Debug> DebugController;
@@ -377,7 +379,11 @@ void FD3D12Adapter::CreateRootDevice(bool bWithDebug)
 	}
 
 	FGenericCrashContext::SetEngineData(TEXT("RHI.D3DDebug"), bWithDebug ? TEXT("true") : TEXT("false"));
+#if PLATFORM_WINDOWS
 	UE_LOG(LogD3D12RHI, Log, TEXT("InitD3DDevice: -D3DDebug = %s -D3D12GPUValidation = %s"), bWithDebug ? TEXT("on") : TEXT("off"), bWithGPUValidation ? TEXT("on") : TEXT("off"));
+#else
+	UE_LOG(LogD3D12RHI, Log, TEXT("InitD3DDevice: -D3DDebug = %s -D3D12GPUValidation = off"), bWithDebug ? TEXT("on") : TEXT("off"));
+#endif
 #endif
 
 #if PLATFORM_WINDOWS
