Fixed Index Buffer SRV NumElements calculation in FD3D12RayTracingGeometry::AllocateBufferSRVs when Index Buffer Stride != 4.

diff --git a/Engine/Source/Runtime/D3D12RHI/Private/D3D12RayTracing.cpp b/Engine/Source/Runtime/D3D12RHI/Private/D3D12RayTracing.cpp
--- a/Engine/Source/Runtime/D3D12RHI/Private/D3D12RayTracing.cpp
+++ b/Engine/Source/Runtime/D3D12RHI/Private/D3D12RayTracing.cpp
@@ -3041,7 +3041,7 @@ void FD3D12RayTracingGeometry::AllocateBufferSRVs(uint32 InGPUIndex)
 		SRVDesc.Format = DXGI_FORMAT_R32_TYPELESS;
 		SRVDesc.Shader4ComponentMapping = D3D12_DEFAULT_SHADER_4_COMPONENT_MAPPING;
 		SRVDesc.Buffer.FirstElement = (Initializer.IndexBufferOffset + IndexBuffer->ResourceLocation.GetOffsetFromBaseOfResource()) >> 2u;
-		SRVDesc.Buffer.NumElements = FMath::Max((uint32)1, Initializer.TotalPrimitiveCount * 3);
+		SRVDesc.Buffer.NumElements = FMath::Max((uint32)1, (Initializer.TotalPrimitiveCount * 3 * IndexBuffer->GetStride()) >> 2u);
 		SRVDesc.Buffer.Flags = D3D12_BUFFER_SRV_FLAG_RAW;
 		SRVDesc.Buffer.StructureByteStride = 0;
 
