D3D12RHI: Add support for NVidia SER extensions

#rnx

diff --git a/Engine/Source/Runtime/D3D12RHI/Private/D3D12NvidiaExtensions.h b/Engine/Source/Runtime/D3D12RHI/Private/D3D12NvidiaExtensions.h
--- a/Engine/Source/Runtime/D3D12RHI/Private/D3D12NvidiaExtensions.h
+++ b/Engine/Source/Runtime/D3D12RHI/Private/D3D12NvidiaExtensions.h
@@ -14,6 +14,12 @@
 	#include "Microsoft/HideMicrosoftPlatformTypes.h"
 #endif // WITH_NVAPI
 
+#if WITH_NVAPI
+	// NVidia Shader Execution Reordering extension
+	#define NV_SER_EXTENSION_SLOT  0
+	#define NV_SER_EXTENSION_SPACE 1001
+#endif
+
 #if NV_AFTERMATH
 
 	struct GFSDK_Aftermath_ContextHandle__;
diff --git a/Engine/Source/Runtime/D3D12RHI/Private/D3D12RayTracing.cpp b/Engine/Source/Runtime/D3D12RHI/Private/D3D12RayTracing.cpp
--- a/Engine/Source/Runtime/D3D12RHI/Private/D3D12RayTracing.cpp
+++ b/Engine/Source/Runtime/D3D12RHI/Private/D3D12RayTracing.cpp
@@ -660,8 +660,28 @@ static TRefCountPtr<ID3D12StateObject> CreateRayTracingStateObject(
 	Desc.pSubobjects = &Subobjects[0];
 	Desc.Type = StateObjectType;
 
+#if WITH_NVAPI
+	bool bReorderingDisabled = false;
+	if (GRHIGlobals.SupportsShaderExecutionReordering && IsRHIDeviceNVIDIA())
+	{
+		// Enable the NVAPI extension to allow shader extensions in ray tracing for shader execution reordering
+		NvAPI_Status NvapiStatus = NvAPI_D3D12_SetNvShaderExtnSlotSpaceLocalThread(RayTracingDevice, NV_SER_EXTENSION_SLOT, NV_SER_EXTENSION_SPACE);
+		checkf(NvapiStatus == NVAPI_OK, TEXT("NVAPI ERROR %d on setup extensions\n"), NvapiStatus);
+		bReorderingDisabled = (NvapiStatus == NVAPI_OK);
+	}
+#endif
+
 	VERIFYD3D12RESULT(RayTracingDevice->CreateStateObject(&Desc, IID_PPV_ARGS(Result.GetInitReference())));
 
+#if WITH_NVAPI
+	if (bReorderingDisabled)
+	{
+		// Disable the NVAPI extension slot again after state object creation.
+		NvAPI_Status NvapiStatus = NvAPI_D3D12_SetNvShaderExtnSlotSpaceLocalThread(RayTracingDevice, ~0u, 0);
+		checkf(NvapiStatus == NVAPI_OK, TEXT("NVAPI ERROR %d on clear extensions\n"), NvapiStatus);
+	}
+#endif
+
 	INC_DWORD_STAT(STAT_D3D12RayTracingCreatedPipelines);
 	INC_DWORD_STAT_BY(STAT_D3D12RayTracingCompiledShaders, NumExports);
 
diff --git a/Engine/Source/Runtime/D3D12RHI/Private/D3D12RootSignature.cpp b/Engine/Source/Runtime/D3D12RHI/Private/D3D12RootSignature.cpp
--- a/Engine/Source/Runtime/D3D12RHI/Private/D3D12RootSignature.cpp
+++ b/Engine/Source/Runtime/D3D12RHI/Private/D3D12RootSignature.cpp
@@ -402,6 +402,19 @@ FD3D12RootSignatureDesc::FD3D12RootSignatureDesc(const FD3D12QuantizedBoundShade
 		RootParametersSize += RootDescriptorCost;
 	}
 
+#if WITH_NVAPI
+	if (QBSS.RootSignatureType == RS_RayTracingGlobal && GRHIGlobals.SupportsShaderExecutionReordering && IsRHIDeviceNVIDIA())
+	{
+		check(RootParameterCount < MaxRootParameters);
+
+		// Should this use the InitAsUnorderedAccessView???
+		DescriptorRanges[RootParameterCount].Init(D3D12_DESCRIPTOR_RANGE_TYPE_UAV, 1, NV_SER_EXTENSION_SLOT, NV_SER_EXTENSION_SPACE, D3D12_DESCRIPTOR_RANGE_FLAG_NONE, 0);
+		TableSlots[RootParameterCount].InitAsDescriptorTable(1, &DescriptorRanges[RootParameterCount]);
+		RootParameterCount++;
+		RootParametersSize += RootDescriptorCost;
+	}
+#endif
+
 	if (bUseShaderDiagnosticBuffer)
 	{
 		check(RootParameterCount < MaxRootParameters);
diff --git a/Engine/Source/Runtime/D3D12RHI/Private/Windows/WindowsD3D12Device.cpp b/Engine/Source/Runtime/D3D12RHI/Private/Windows/WindowsD3D12Device.cpp
--- a/Engine/Source/Runtime/D3D12RHI/Private/Windows/WindowsD3D12Device.cpp
+++ b/Engine/Source/Runtime/D3D12RHI/Private/Windows/WindowsD3D12Device.cpp
@@ -1475,6 +1475,31 @@ void FD3D12DynamicRHI::Init()
 			{
 				UE_LOG(LogD3D12RHI, Warning, TEXT("Failed to query support for 64 bit atomics"));
 			}
+
+			NvAPI_Status NvStatusSER = NvAPI_D3D12_IsNvShaderExtnOpCodeSupported(GetAdapter().GetD3DDevice(), NV_EXTN_OP_HIT_OBJECT_REORDER_THREAD, &GRHIGlobals.SupportsShaderExecutionReordering);
+			if (NvStatusSER == NVAPI_OK && GRHIGlobals.SupportsShaderExecutionReordering)
+			{
+				NVAPI_D3D12_RAYTRACING_THREAD_REORDERING_CAPS ReorderCaps = NVAPI_D3D12_RAYTRACING_THREAD_REORDERING_CAP_NONE;
+				NvStatusSER = NvAPI_D3D12_GetRaytracingCaps(
+					GetAdapter().GetD3DDevice(),
+					NVAPI_D3D12_RAYTRACING_CAPS_TYPE_THREAD_REORDERING,
+					&ReorderCaps,
+					sizeof(ReorderCaps));
+
+				if (NvStatusSER != NVAPI_OK || ReorderCaps == NVAPI_D3D12_RAYTRACING_THREAD_REORDERING_CAP_NONE)
+				{
+					GRHIGlobals.SupportsShaderExecutionReordering = false;
+				}
+			}
+
+			if (GRHIGlobals.SupportsShaderExecutionReordering)
+			{
+				UE_LOG(LogD3D12RHI, Log, TEXT("NVIDIA Shader Execution Reordering interface supported!"));
+			}
+			else
+			{
+				UE_LOG(LogD3D12RHI, Log, TEXT("NVIDIA Shader Execution Reordering NOT supported!"));
+			}
 		}
 		else
 		{
diff --git a/Engine/Source/Runtime/RHI/Public/RHIGlobals.h b/Engine/Source/Runtime/RHI/Public/RHIGlobals.h
--- a/Engine/Source/Runtime/RHI/Public/RHIGlobals.h
+++ b/Engine/Source/Runtime/RHI/Public/RHIGlobals.h
@@ -680,6 +680,9 @@ struct FRHIGlobals
 
 	/** True if RHI supports Linear texture format in 3D/Cube/Array texture */
 	bool SupportLinearTextureVolumeFormat = true;
+
+	/** true if the RHI supports shader execution reordering. */
+	bool SupportsShaderExecutionReordering = false;
 };
 
 extern RHI_API FRHIGlobals GRHIGlobals;
