Implement resource renamed on TD3D12View instead of FD3D12View so UpdateView can be called which handles empty resources on stream out

diff --git a/Engine/Source/Runtime/D3D12RHI/Private/D3D12View.cpp b/Engine/Source/Runtime/D3D12RHI/Private/D3D12View.cpp
--- a/Engine/Source/Runtime/D3D12RHI/Private/D3D12View.cpp
+++ b/Engine/Source/Runtime/D3D12RHI/Private/D3D12View.cpp
@@ -323,18 +323,6 @@ void FD3D12View::UpdateView(FD3D12ContextArray const& Contexts, const FResourceI
 	UpdateBindlessSlot(Contexts);
 }
 
-void FD3D12View::ResourceRenamed(FD3D12ContextArray const& Contexts, FD3D12BaseShaderResource* InRenamedResource, FD3D12ResourceLocation* InNewResourceLocation)
-{
-	// Can only be called if the base shader resource is not null.
-	check(ResourceInfo.BaseResource == InRenamedResource && ResourceInfo.ResourceLocation == InNewResourceLocation);
-
-	// Update the cached resource pointers
-	ResourceInfo = InRenamedResource;
-
-	UpdateDescriptor();
-	UpdateBindlessSlot(Contexts);
-}
-
 
 
 // -----------------------------------------------------------------------------------------------------
diff --git a/Engine/Source/Runtime/D3D12RHI/Private/D3D12View.h b/Engine/Source/Runtime/D3D12RHI/Private/D3D12View.h
--- a/Engine/Source/Runtime/D3D12RHI/Private/D3D12View.h
+++ b/Engine/Source/Runtime/D3D12RHI/Private/D3D12View.h
@@ -221,8 +221,7 @@ public:
 protected:
 	FD3D12View(FD3D12Device* InDevice, ERHIDescriptorHeapType InHeapType, FD3D12View* FirstLinkedObject);
 	virtual ~FD3D12View();
-
-	virtual void ResourceRenamed(FD3D12ContextArray const& Contexts, FD3D12BaseShaderResource* InRenamedResource, FD3D12ResourceLocation* InNewResourceLocation) override;
+		
 	virtual void UpdateDescriptor() = 0;
 
 	void UpdateResourceInfo(FResourceInfo const& InResource, FNullDescPtr NullDescriptor);
@@ -272,6 +271,14 @@ protected:
 		FD3D12View::UpdateView(Contexts, InResource, TParent::Null);
 	}
 
+	virtual void ResourceRenamed(FD3D12ContextArray const& Contexts, FD3D12BaseShaderResource* InRenamedResource, FD3D12ResourceLocation* InNewResourceLocation) override
+	{
+		// Can only be called if the base shader resource is not null.
+		check(ResourceInfo.BaseResource == InRenamedResource && ResourceInfo.ResourceLocation == InNewResourceLocation);
+
+		UpdateView(Contexts, InRenamedResource, D3DViewDesc);
+	}
+	
 public:
 	TDesc const& GetD3DDesc() const { return D3DViewDesc; }
 };
