Also clean up RAW index & vertex buffer SRVs when calling ReleaseUnderlyingResource on ray tracing geometry (called during stream out)

diff --git a/Engine/Source/Runtime/D3D12RHI/Private/D3D12RayTracing.cpp b/Engine/Source/Runtime/D3D12RHI/Private/D3D12RayTracing.cpp
--- a/Engine/Source/Runtime/D3D12RHI/Private/D3D12RayTracing.cpp
+++ b/Engine/Source/Runtime/D3D12RHI/Private/D3D12RayTracing.cpp
@@ -2967,7 +2967,7 @@ void FD3D12RayTracingGeometry::Swap(FD3D12RayTracingGeometry& Other)
 
 	GeometryDescs.SetNumUninitialized(Initializer.Segments.Num());
 	TranslateRayTracingGeometryDescs(Initializer, GeometryDescs);
-
+		
 	FOREACH_GPU(GPUIndex < MAX_NUM_GPUS && GPUIndex < GNumExplicitGPUsForRendering,
 	{
 		RegisterAsRenameListener(GPUIndex);
@@ -2992,9 +2992,12 @@ void FD3D12RayTracingGeometry::ReleaseUnderlyingResource()
 		}
 	});
 
-	// Unregister as dependent resource on vertex and index buffers
+	// Unregister as dependent resource on vertex and index buffers & clear the SRVs
 	FOREACH_GPU(GPUIndex < MAX_NUM_GPUS,
 	{
+		HitGroupSystemIndexBufferSRV[GPUIndex].Reset();
+		HitGroupSystemSegmentVertexBufferSRVs[GPUIndex].Empty();
+
 		UnregisterAsRenameListener(GPUIndex);
 	});
 
