[HWRT] Provide FRayTracingSceneInitializer2 to RHICalcRayTracingSceneSize(...) instead of individual parameters.

diff --git a/Engine/Plugins/Experimental/GPULightmass/Source/GPULightmass/Private/LightmapRenderer.cpp b/Engine/Plugins/Experimental/GPULightmass/Source/GPULightmass/Private/LightmapRenderer.cpp
--- a/Engine/Plugins/Experimental/GPULightmass/Source/GPULightmass/Private/LightmapRenderer.cpp
+++ b/Engine/Plugins/Experimental/GPULightmass/Source/GPULightmass/Private/LightmapRenderer.cpp
@@ -1162,7 +1162,7 @@ bool FSceneRenderState::SetupRayTracingScene(FRDGBuilder& GraphBuilder, FSceneUn
 
 			const FRayTracingSceneInitializer2& SceneInitializer = RayTracingScene->GetInitializer();
 
-			FRayTracingAccelerationStructureSize SizeInfo = RHICalcRayTracingSceneSize(SceneInitializer.MaxNumInstances, SceneBuildFlags);
+			FRayTracingAccelerationStructureSize SizeInfo = RHICalcRayTracingSceneSize(SceneInitializer);
 			FRHIResourceCreateInfo BufferCreateInfo(TEXT("LightmassRayTracingSceneBuffer"));
 			RayTracingSceneBuffer = RHICmdList.CreateBuffer(uint32(SizeInfo.ResultSize), BUF_AccelerationStructure, 0, ERHIAccess::BVHWrite, BufferCreateInfo);
 			RayTracingSceneSRV = RHICmdList.CreateShaderResourceView(FShaderResourceViewInitializer(RayTracingSceneBuffer, RayTracingScene, 0));
diff --git a/Engine/Source/Runtime/Apple/MetalRHI/Private/MetalDynamicRHI.h b/Engine/Source/Runtime/Apple/MetalRHI/Private/MetalDynamicRHI.h
--- a/Engine/Source/Runtime/Apple/MetalRHI/Private/MetalDynamicRHI.h
+++ b/Engine/Source/Runtime/Apple/MetalRHI/Private/MetalDynamicRHI.h
@@ -155,7 +155,7 @@ public:
 	virtual uint16 RHIGetPlatformTextureMaxSampleCount() override;
 
 #if METAL_RHI_RAYTRACING
-	virtual FRayTracingAccelerationStructureSize RHICalcRayTracingSceneSize(uint32 MaxInstances, ERayTracingAccelerationStructureFlags Flags) final override;
+	virtual FRayTracingAccelerationStructureSize RHICalcRayTracingSceneSize(const FRayTracingSceneInitializer2& Initializer) final override;
 	virtual FRayTracingAccelerationStructureSize RHICalcRayTracingGeometrySize(const FRayTracingGeometryInitializer& Initializer) final override;
 
 	virtual FRayTracingGeometryRHIRef RHICreateRayTracingGeometry(FRHICommandListBase& RHICmdList, const FRayTracingGeometryInitializer& Initializer) final override;
diff --git a/Engine/Source/Runtime/Apple/MetalRHI/Private/MetalRayTracing.cpp b/Engine/Source/Runtime/Apple/MetalRHI/Private/MetalRayTracing.cpp
--- a/Engine/Source/Runtime/Apple/MetalRHI/Private/MetalRayTracing.cpp
+++ b/Engine/Source/Runtime/Apple/MetalRHI/Private/MetalRayTracing.cpp
@@ -314,14 +314,14 @@ FRayTracingAccelerationStructureSize FMetalDynamicRHI::RHICalcRayTracingGeometry
 	return CalcRayTracingGeometrySize(AccelerationStructureDescriptor);
 }
 
-FRayTracingAccelerationStructureSize FMetalDynamicRHI::RHICalcRayTracingSceneSize(uint32 MaxInstances, ERayTracingAccelerationStructureFlags Flags)
+FRayTracingAccelerationStructureSize FMetalDynamicRHI::RHICalcRayTracingSceneSize(const FRayTracingSceneInitializer2& Initializer)
 {
 	// TODO: Do we need to take in account the flags provided by the function call?
 	// TODO: Can we get away with the instance count only? (works on AS; what about AMD?)
     MTL_SCOPED_AUTORELEASE_POOL;
     
     MTL::InstanceAccelerationStructureDescriptor InstanceDescriptor = MTL::InstanceAccelerationStructureDescriptor();
-    InstanceDescriptor.SetInstanceCount(MaxInstances);
+    InstanceDescriptor.SetInstanceCount(Initializer.MaxNumInstances);
 
     return CalcRayTracingGeometrySize(InstanceDescriptor);
 }
diff --git a/Engine/Source/Runtime/D3D12RHI/Private/D3D12RHIPrivate.h b/Engine/Source/Runtime/D3D12RHI/Private/D3D12RHIPrivate.h
--- a/Engine/Source/Runtime/D3D12RHI/Private/D3D12RHIPrivate.h
+++ b/Engine/Source/Runtime/D3D12RHI/Private/D3D12RHIPrivate.h
@@ -466,7 +466,7 @@ public:
 
 #if D3D12_RHI_RAYTRACING
 
-	virtual FRayTracingAccelerationStructureSize RHICalcRayTracingSceneSize(uint32 MaxInstances, ERayTracingAccelerationStructureFlags Flags) final override;
+	virtual FRayTracingAccelerationStructureSize RHICalcRayTracingSceneSize(const FRayTracingSceneInitializer2& Initializer) final override;
 	virtual FRayTracingAccelerationStructureSize RHICalcRayTracingGeometrySize(const FRayTracingGeometryInitializer& Initializer) final override;
 
 	virtual FRayTracingGeometryRHIRef RHICreateRayTracingGeometry(FRHICommandListBase& RHICmdList, const FRayTracingGeometryInitializer& Initializer) final override;
diff --git a/Engine/Source/Runtime/D3D12RHI/Private/D3D12RayTracing.cpp b/Engine/Source/Runtime/D3D12RHI/Private/D3D12RayTracing.cpp
--- a/Engine/Source/Runtime/D3D12RHI/Private/D3D12RayTracing.cpp
+++ b/Engine/Source/Runtime/D3D12RHI/Private/D3D12RayTracing.cpp
@@ -2448,13 +2448,13 @@ static D3D12_RAYTRACING_INSTANCE_FLAGS TranslateRayTracingInstanceFlags(ERayTrac
 	return Result;
 }
 
-FRayTracingAccelerationStructureSize FD3D12DynamicRHI::RHICalcRayTracingSceneSize(uint32 MaxInstances, ERayTracingAccelerationStructureFlags Flags)
+FRayTracingAccelerationStructureSize FD3D12DynamicRHI::RHICalcRayTracingSceneSize(const FRayTracingSceneInitializer2& Initializer)
 {
 	D3D12_BUILD_RAYTRACING_ACCELERATION_STRUCTURE_INPUTS BuildInputs = {};
 	BuildInputs.Type = D3D12_RAYTRACING_ACCELERATION_STRUCTURE_TYPE_TOP_LEVEL;
 	BuildInputs.DescsLayout = D3D12_ELEMENTS_LAYOUT_ARRAY;
-	BuildInputs.NumDescs = MaxInstances;
-	BuildInputs.Flags = TranslateRayTracingAccelerationStructureFlags(Flags);
+	BuildInputs.NumDescs = Initializer.MaxNumInstances;
+	BuildInputs.Flags = TranslateRayTracingAccelerationStructureFlags(Initializer.BuildFlags);
 
 	FD3D12Adapter& Adapter = GetAdapter();
 
@@ -3518,7 +3518,7 @@ FD3D12RayTracingScene::FD3D12RayTracingScene(FD3D12Adapter* Adapter, FRayTracing
 	checkf(Initializer.Lifetime == RTSL_SingleFrame, TEXT("Only single-frame ray tracing scenes are currently implemented."));
 
 	// Get maximum buffer sizes for all GPUs in the system
-	SizeInfo = RHICalcRayTracingSceneSize(Initializer.MaxNumInstances, Initializer.BuildFlags);
+	SizeInfo = RHICalcRayTracingSceneSize(Initializer);
 };
 
 FD3D12RayTracingScene::~FD3D12RayTracingScene()
diff --git a/Engine/Source/Runtime/RHI/Public/DynamicRHI.h b/Engine/Source/Runtime/RHI/Public/DynamicRHI.h
--- a/Engine/Source/Runtime/RHI/Public/DynamicRHI.h
+++ b/Engine/Source/Runtime/RHI/Public/DynamicRHI.h
@@ -949,7 +949,7 @@ public:
 
 	virtual uint16 RHIGetPlatformTextureMaxSampleCount() { return 8; };
 
-	virtual FRayTracingAccelerationStructureSize RHICalcRayTracingSceneSize(uint32 MaxInstances, ERayTracingAccelerationStructureFlags Flags)
+	virtual FRayTracingAccelerationStructureSize RHICalcRayTracingSceneSize(const FRayTracingSceneInitializer2& Initializer)
 	{
 		checkNoEntry();
 		return {};
@@ -1467,9 +1467,19 @@ FORCEINLINE bool RHIMatchPrecachePSOInitializers(const FGraphicsPipelineStateIni
 	return GDynamicRHI->RHIMatchPrecachePSOInitializers(LHS, RHS);
 }
 
+FORCEINLINE FRayTracingAccelerationStructureSize RHICalcRayTracingSceneSize(const FRayTracingSceneInitializer2& Initializer)
+{
+	return GDynamicRHI->RHICalcRayTracingSceneSize(Initializer);
+}
+
+UE_DEPRECATED(5.5, "Provide FRayTracingSceneInitializer2 instead.")
 FORCEINLINE FRayTracingAccelerationStructureSize RHICalcRayTracingSceneSize(uint32 MaxInstances, ERayTracingAccelerationStructureFlags Flags)
 {
-	return GDynamicRHI->RHICalcRayTracingSceneSize(MaxInstances, Flags);
+	FRayTracingSceneInitializer2 Initializer;
+	Initializer.MaxNumInstances = MaxInstances;
+	Initializer.BuildFlags = Flags;
+
+	return GDynamicRHI->RHICalcRayTracingSceneSize(Initializer);
 }
 
 FORCEINLINE FRayTracingSceneRHIRef RHICreateRayTracingScene(FRayTracingSceneInitializer2 Initializer)
diff --git a/Engine/Source/Runtime/RHI/Public/RHIValidation.h b/Engine/Source/Runtime/RHI/Public/RHIValidation.h
--- a/Engine/Source/Runtime/RHI/Public/RHIValidation.h
+++ b/Engine/Source/Runtime/RHI/Public/RHIValidation.h
@@ -1230,9 +1230,9 @@ public:
 		return RHI->RHICreateRayTracingPipelineState(Initializer);
 	}
 
-	virtual FRayTracingAccelerationStructureSize RHICalcRayTracingSceneSize(uint32 MaxInstances, ERayTracingAccelerationStructureFlags Flags) override final
+	virtual FRayTracingAccelerationStructureSize RHICalcRayTracingSceneSize(const FRayTracingSceneInitializer2& Initializer) override final
 	{
-		return RHI->RHICalcRayTracingSceneSize(MaxInstances, Flags);
+		return RHI->RHICalcRayTracingSceneSize(Initializer);
 	}
 
 	virtual FRayTracingAccelerationStructureSize RHICalcRayTracingGeometrySize(const FRayTracingGeometryInitializer& Initializer) override final
diff --git a/Engine/Source/Runtime/Renderer/Private/Tests/RayTracingTestbed.cpp b/Engine/Source/Runtime/Renderer/Private/Tests/RayTracingTestbed.cpp
--- a/Engine/Source/Runtime/Renderer/Private/Tests/RayTracingTestbed.cpp
+++ b/Engine/Source/Runtime/Renderer/Private/Tests/RayTracingTestbed.cpp
@@ -193,7 +193,7 @@ bool RunRayTracingTestbed_RenderThread(const FString& Parameters)
 
 	const FRayTracingSceneInitializer2& SceneInitializer = RayTracingSceneRHI->GetInitializer();
 
-	FRayTracingAccelerationStructureSize SceneSizeInfo = RHICalcRayTracingSceneSize(NumInstances, SceneBuildFlags);
+	FRayTracingAccelerationStructureSize SceneSizeInfo = RHICalcRayTracingSceneSize(SceneInitializer);
 	FRHIResourceCreateInfo SceneBufferCreateInfo(TEXT("RayTracingTestBedSceneBuffer"));
 	FBufferRHIRef SceneBuffer = RHICmdList.CreateBuffer(
 		uint32(SceneSizeInfo.ResultSize),
diff --git a/Engine/Source/Runtime/VulkanRHI/Private/VulkanRayTracing.cpp b/Engine/Source/Runtime/VulkanRHI/Private/VulkanRayTracing.cpp
--- a/Engine/Source/Runtime/VulkanRHI/Private/VulkanRayTracing.cpp
+++ b/Engine/Source/Runtime/VulkanRHI/Private/VulkanRayTracing.cpp
@@ -708,7 +708,7 @@ FVulkanRayTracingScene::FVulkanRayTracingScene(FRayTracingSceneInitializer2 InIn
 {
 	INC_DWORD_STAT(STAT_VulkanRayTracingAllocatedTLAS);
 
-	SizeInfo = RHICalcRayTracingSceneSize(Initializer.MaxNumInstances, Initializer.BuildFlags);
+	SizeInfo = RHICalcRayTracingSceneSize(Initializer);
 }
 
 FVulkanRayTracingScene::~FVulkanRayTracingScene()
@@ -1138,11 +1138,11 @@ FRHIShaderBindingTable* FVulkanRayTracingScene::FindOrCreateShaderBindingTable(c
 	return CreatedShaderTable;
 }
 
-FRayTracingAccelerationStructureSize FVulkanDynamicRHI::RHICalcRayTracingSceneSize(uint32 MaxInstances, ERayTracingAccelerationStructureFlags Flags)
+FRayTracingAccelerationStructureSize FVulkanDynamicRHI::RHICalcRayTracingSceneSize(const FRayTracingSceneInitializer2& Initializer)
 {
 	FVkRtTLASBuildData BuildData;
 	const VkDeviceAddress InstanceBufferAddress = 0; // No device address available when only querying TLAS size
-	GetTLASBuildData(Device->GetInstanceHandle(), MaxInstances, InstanceBufferAddress, Flags, EAccelerationStructureBuildMode::Build, BuildData);
+	GetTLASBuildData(Device->GetInstanceHandle(), Initializer.MaxNumInstances, InstanceBufferAddress, Initializer.BuildFlags, EAccelerationStructureBuildMode::Build, BuildData);
 
 	FRayTracingAccelerationStructureSize Result;
 	Result.ResultSize = BuildData.SizesInfo.accelerationStructureSize;
diff --git a/Engine/Source/Runtime/VulkanRHI/Public/VulkanDynamicRHI.h b/Engine/Source/Runtime/VulkanRHI/Public/VulkanDynamicRHI.h
--- a/Engine/Source/Runtime/VulkanRHI/Public/VulkanDynamicRHI.h
+++ b/Engine/Source/Runtime/VulkanRHI/Public/VulkanDynamicRHI.h
@@ -232,7 +232,7 @@ public:
 	virtual FRHIResourceCollectionRef RHICreateResourceCollection(FRHICommandListBase& RHICmdList, TConstArrayView<FRHIResourceCollectionMember> InMembers) final;
 #endif
 
-	virtual FRayTracingAccelerationStructureSize RHICalcRayTracingSceneSize(uint32 MaxInstances, ERayTracingAccelerationStructureFlags Flags) final override;
+	virtual FRayTracingAccelerationStructureSize RHICalcRayTracingSceneSize(const FRayTracingSceneInitializer2& Initializer) final override;
 	virtual FRayTracingAccelerationStructureSize RHICalcRayTracingGeometrySize(const FRayTracingGeometryInitializer& Initializer) final override;
 	virtual FRayTracingGeometryRHIRef RHICreateRayTracingGeometry(FRHICommandListBase& RHICmdList, const FRayTracingGeometryInitializer& Initializer) final override;
 	virtual FRayTracingSceneRHIRef RHICreateRayTracingScene(FRayTracingSceneInitializer2 Initializer) final override;
