D3D12 Bindless Cleanup
- Renaming FD3D12BindlessManagerAdapter to FD3D12BindlessDescriptorAllocator to better represent its function.
- FD3D12BindlessDescriptorManager and subclasses now hold a reference to their Adapter's FD3D12BindlessDescriptorAllocator.
- Removing a few redundant/unnecessary variables and methods.
- Adding LLM markers for bindless heaps.

diff --git a/Engine/Source/Runtime/D3D12RHI/Private/D3D12Adapter.cpp b/Engine/Source/Runtime/D3D12RHI/Private/D3D12Adapter.cpp
--- a/Engine/Source/Runtime/D3D12RHI/Private/D3D12Adapter.cpp
+++ b/Engine/Source/Runtime/D3D12RHI/Private/D3D12Adapter.cpp
@@ -318,6 +318,9 @@ FD3D12Adapter::FD3D12Adapter(FD3D12AdapterDesc& DescIn)
 	, StaticRayTracingGlobalRootSignature(this)
 	, StaticRayTracingLocalRootSignature(this)
 #endif
+#if PLATFORM_SUPPORTS_BINDLESS_RENDERING
+	, BindlessDescriptorAllocator(this)
+#endif
 {
 	FMemory::Memzero(&UploadHeapAllocator, sizeof(UploadHeapAllocator));
 	FMemory::Memzero(&Devices, sizeof(Devices));
@@ -1200,10 +1203,10 @@ void FD3D12Adapter::InitializeDevices()
 		if (Desc.MaxSupportedFeatureLevel >= D3D_FEATURE_LEVEL_12_0 && Desc.MaxSupportedShaderModel >= D3D_SHADER_MODEL_6_6 && Desc.ResourceBindingTier >= D3D12_RESOURCE_BINDING_TIER_3)
 		{
 			// Needs to happen before device creation below
-			BindlessManager.Init(this);
+			BindlessDescriptorAllocator.Init();
 
-			BindlessResourcesConfig = BindlessManager.GetResourcesConfiguration();
-			BindlessSamplersConfig = BindlessManager.GetSamplersConfiguration();
+			BindlessResourcesConfig = BindlessDescriptorAllocator.GetResourcesConfiguration();
+			BindlessSamplersConfig = BindlessDescriptorAllocator.GetSamplersConfiguration();
 
 			GRHIGlobals.ShaderBundles.RequiresSharedBindlessParameters = (BindlessResourcesConfig == ERHIBindlessConfiguration::AllShaders || BindlessSamplersConfig == ERHIBindlessConfiguration::AllShaders);
 		}
diff --git a/Engine/Source/Runtime/D3D12RHI/Private/D3D12Adapter.h b/Engine/Source/Runtime/D3D12RHI/Private/D3D12Adapter.h
--- a/Engine/Source/Runtime/D3D12RHI/Private/D3D12Adapter.h
+++ b/Engine/Source/Runtime/D3D12RHI/Private/D3D12Adapter.h
@@ -449,7 +449,9 @@ public:
 	HMODULE GetDxgiDllHandle() const { return DxgiDllHandle; };
 #endif
 
-	FD3D12BindlessManagerAdapter& GetBindlessManager() { return BindlessManager; }
+#if PLATFORM_SUPPORTS_BINDLESS_RENDERING
+	FD3D12BindlessDescriptorAllocator& GetBindlessDescriptorAllocator() { return BindlessDescriptorAllocator; }
+#endif
 
 protected:
 
@@ -620,7 +622,9 @@ protected:
 	FD3D12RootSignature StaticRayTracingLocalRootSignature;
 #endif
 
-	FD3D12BindlessManagerAdapter BindlessManager;
+#if PLATFORM_SUPPORTS_BINDLESS_RENDERING
+	FD3D12BindlessDescriptorAllocator BindlessDescriptorAllocator;
+#endif
 
 private:
 	// Insight memory trace helper
diff --git a/Engine/Source/Runtime/D3D12RHI/Private/D3D12BindlessDescriptors.cpp b/Engine/Source/Runtime/D3D12RHI/Private/D3D12BindlessDescriptors.cpp
--- a/Engine/Source/Runtime/D3D12RHI/Private/D3D12BindlessDescriptors.cpp
+++ b/Engine/Source/Runtime/D3D12RHI/Private/D3D12BindlessDescriptors.cpp
@@ -31,6 +31,8 @@ static FAutoConsoleVariableRef CVarBindlessSamplerDescriptorHeapSize(
 
 FD3D12DescriptorHeap* UE::D3D12BindlessDescriptors::CreateCpuHeap(FD3D12Device* InDevice, ERHIDescriptorHeapType InType, uint32 InNewNumDescriptorsPerHeap)
 {
+	LLM_SCOPE_BYNAME(TEXT("RHIMisc/BindlessDescriptorHeap/CPU"));
+
 	const TCHAR* const HeapName = (InType == ERHIDescriptorHeapType::Standard) ? TEXT("BindlessResourcesCPU") : TEXT("BindlessSamplersCPU");
 
 	return InDevice->GetDescriptorHeapManager().AllocateIndependentHeap(
@@ -43,6 +45,7 @@ FD3D12DescriptorHeap* UE::D3D12BindlessDescriptors::CreateCpuHeap(FD3D12Device*
 
 FD3D12DescriptorHeap* UE::D3D12BindlessDescriptors::CreateGpuHeap(FD3D12Device* InDevice, ERHIDescriptorHeapType InType, uint32 InNewNumDescriptorsPerHeap)
 {
+	LLM_SCOPE_BYNAME(TEXT("RHIMisc/BindlessDescriptorHeap/GPU"));
 	SCOPED_NAMED_EVENT_F(TEXT("CreateNewBindlessHeap (%d)"), FColor::Turquoise, InNewNumDescriptorsPerHeap);
 
 	const TCHAR* const HeapName = (InType == ERHIDescriptorHeapType::Standard) ? TEXT("BindlessResources") : TEXT("BindlessSamplers");
@@ -63,12 +66,11 @@ void UE::D3D12BindlessDescriptors::DeferredFreeHeap(FD3D12Device* InDevice, FD3D
 ////////////////////////////////////////////////////////////////////////////////////////////////////////////////////////
 // FD3D12BindlessSamplerManager
 
-FD3D12BindlessSamplerManager::FD3D12BindlessSamplerManager(FD3D12Device* InDevice, ERHIBindlessConfiguration InConfiguration, uint32 InNumDescriptors, TConstArrayView<TStatId> InStats)
+FD3D12BindlessSamplerManager::FD3D12BindlessSamplerManager(FD3D12Device* InDevice, FD3D12BindlessDescriptorAllocator& InAllocator)
 	: FD3D12DeviceChild(InDevice)
-	, Allocator(ERHIDescriptorHeapType::Sampler, InNumDescriptors, InStats)
-	, Configuration(InConfiguration)
+	, GpuHeap(UE::D3D12BindlessDescriptors::CreateGpuHeap(InDevice, ERHIDescriptorHeapType::Sampler, InAllocator.GetSamplerCapacity()))
+	, Configuration(InAllocator.GetSamplersConfiguration())
 {
-	GpuHeap = UE::D3D12BindlessDescriptors::CreateGpuHeap(InDevice, ERHIDescriptorHeapType::Sampler, InNumDescriptors);
 }
 
 void FD3D12BindlessSamplerManager::CleanupResources()
@@ -76,7 +78,7 @@ void FD3D12BindlessSamplerManager::CleanupResources()
 	GpuHeap = nullptr;
 }
 
-void FD3D12BindlessSamplerManager::InitializeSampler(FRHIDescriptorHandle DstHandle, FD3D12SamplerState* SamplerState)
+void FD3D12BindlessSamplerManager::InitializeDescriptor(FRHIDescriptorHandle DstHandle, FD3D12SamplerState* SamplerState)
 {
 	check(DstHandle.GetType() == ERHIDescriptorHeapType::Sampler);
 
@@ -105,11 +107,16 @@ FD3D12DescriptorHeap* FD3D12BindlessSamplerManager::GetExplicitHeapForContext(FD
 }
 
 ////////////////////////////////////////////////////////////////////////////////////////////////////////////////////////
-// FD3D12BindlessManagerAdapter
+// FD3D12BindlessDescriptorAllocator
 
-void FD3D12BindlessManagerAdapter::Init(FD3D12Adapter* InParent)
+FD3D12BindlessDescriptorAllocator::FD3D12BindlessDescriptorAllocator(FD3D12Adapter* InParent)
+	: FD3D12AdapterChild(InParent)
 {
-	SetParentAdapter(InParent);
+}
+
+void FD3D12BindlessDescriptorAllocator::Init()
+{
+	LLM_SCOPE_BYNAME(TEXT("RHIMisc/BindlessDescriptorAllocator"));
 
 	BindlessResourcesConfiguration = RHIGetRuntimeBindlessResourcesConfiguration(GMaxRHIShaderPlatform);
 	BindlessSamplersConfiguration = RHIGetRuntimeBindlessSamplersConfiguration(GMaxRHIShaderPlatform);
@@ -150,19 +157,16 @@ void FD3D12BindlessManagerAdapter::Init(FD3D12Adapter* InParent)
 
 		SamplerAllocator = new FRHIHeapDescriptorAllocator(ERHIDescriptorHeapType::Sampler, NumSamplerDescriptors, Stats);
 	}
-
-	bBindlessResourcesAllowed = (BindlessResourcesConfiguration != ERHIBindlessConfiguration::Disabled);
-	bBindlessSamplersAllowed = (BindlessSamplersConfiguration != ERHIBindlessConfiguration::Disabled);
 }
 
-FRHIDescriptorHandle FD3D12BindlessManagerAdapter::AllocateSamplerHandle()
+FRHIDescriptorHandle FD3D12BindlessDescriptorAllocator::AllocateSamplerHandle()
 {
 	FRHIDescriptorHandle Result = SamplerAllocator->Allocate();
 	check(Result.IsValid());
 	return Result;
 }
 
-void FD3D12BindlessManagerAdapter::FreeSamplerHandle(FRHIDescriptorHandle InHandle)
+void FD3D12BindlessDescriptorAllocator::FreeSamplerHandle(FRHIDescriptorHandle InHandle)
 {
 	if (InHandle.IsValid())
 	{
@@ -170,17 +174,16 @@ void FD3D12BindlessManagerAdapter::FreeSamplerHandle(FRHIDescriptorHandle InHand
 	}
 }
 
-FRHIDescriptorHandle FD3D12BindlessManagerAdapter::AllocateResourceHandle()
+FRHIDescriptorHandle FD3D12BindlessDescriptorAllocator::AllocateResourceHandle()
 {
-	if (!bBindlessResourcesAllowed)
+	if (!AreResourcesBindless())
 	{
 		return FRHIDescriptorHandle();
 	}
 
 	FRHIDescriptorHandle Result = ResourceAllocator->Allocate();
 
-#if !D3D12RHI_CUSTOM_BINDLESS_RESOURCE_MANAGER
-	// Custom resource manager used on console doesn't support dynamic growth
+#if D3D12RHI_BINDLESS_RESOURCE_MANAGER_SUPPORTS_RESIZING
 	if (!Result.IsValid())
 	{
 		TRACE_CPUPROFILER_EVENT_SCOPE(FD3D12Adapter::BindlessResourceAllocateHandle(GrowHeap));
@@ -195,7 +198,8 @@ FRHIDescriptorHandle FD3D12BindlessManagerAdapter::AllocateResourceHandle()
 		// Grow the CPU heaps for all devices
 		for (FD3D12Device* ParentDevice : GetParentAdapter()->GetDevices())
 		{
-			ParentDevice->GetBindlessDescriptorManager().GetResourceManager().GrowCPUHeap(CurrentNumDescriptors, NewNumDescriptors);
+			checkSlow(ParentDevice->GetBindlessDescriptorManager().GetResourceManager() != nullptr);
+			ParentDevice->GetBindlessDescriptorManager().GetResourceManager()->GrowCPUHeap(CurrentNumDescriptors, NewNumDescriptors);
 		}
 
 		return Result;
@@ -206,7 +210,7 @@ FRHIDescriptorHandle FD3D12BindlessManagerAdapter::AllocateResourceHandle()
 	return Result;
 }
 
-void FD3D12BindlessManagerAdapter::FreeResourceHandle(FRHIDescriptorHandle InHandle)
+void FD3D12BindlessDescriptorAllocator::FreeResourceHandle(FRHIDescriptorHandle InHandle)
 {
 	if (InHandle.IsValid())
 	{
@@ -214,17 +218,16 @@ void FD3D12BindlessManagerAdapter::FreeResourceHandle(FRHIDescriptorHandle InHan
 	}
 }
 
-
 ////////////////////////////////////////////////////////////////////////////////////////////////////////////////////////
 // FD3D12BindlessResourceManager
 
 #if !D3D12RHI_CUSTOM_BINDLESS_RESOURCE_MANAGER
 
-FD3D12BindlessResourceManager::FD3D12BindlessResourceManager(FD3D12Device* InDevice, FD3D12Adapter* InAdapter)
+FD3D12BindlessResourceManager::FD3D12BindlessResourceManager(FD3D12Device* InDevice, FD3D12BindlessDescriptorAllocator& InAllocator)
 	: FD3D12DeviceChild(InDevice)
-	, HeapsCS(InAdapter->GetBindlessManager().GetResourceHeapsCS())
-	, CpuHeap(UE::D3D12BindlessDescriptors::CreateCpuHeap(InDevice, ERHIDescriptorHeapType::Standard, InAdapter->GetBindlessManager().GetResourceCapacity()))
-	, Configuration(InAdapter->GetBindlessManager().GetResourcesConfiguration())
+	, HeapsCS(InAllocator.GetResourceHeapsCS())
+	, CpuHeap(UE::D3D12BindlessDescriptors::CreateCpuHeap(InDevice, ERHIDescriptorHeapType::Standard, InAllocator.GetResourceCapacity()))
+	, Configuration(InAllocator.GetResourcesConfiguration())
 {	
 	if (GetConfiguration() == ERHIBindlessConfiguration::AllShaders)
 	{
@@ -578,7 +581,7 @@ void FD3D12BindlessResourceManager::CopyCpuHeap(FD3D12DescriptorHeap* Destinatio
 {
 	// Copy the smallest possible set of descriptors from the CPU heap to the new GPU heap.
 	FRHIDescriptorAllocatorRange AllocatedRange(0, 0);
-	if (GetParentDevice()->GetParentAdapter()->GetBindlessManager().GetResourceAllocatedRange(AllocatedRange))
+	if (GetParentDevice()->GetBindlessDescriptorAllocator().GetResourceAllocatedRange(AllocatedRange))
 	{
 		const uint32 NumDescriptorsToCopy = AllocatedRange.Last - AllocatedRange.First + 1;
 		UE::D3D12Descriptors::CopyDescriptors(GetParentDevice(), DestinationHeap, CpuHeap, AllocatedRange.First, NumDescriptorsToCopy);
@@ -706,8 +709,9 @@ void FD3D12BindlessResourceManager::CheckRequestNewActiveGPUHeap()
 ////////////////////////////////////////////////////////////////////////////////////////////////////////////////////////
 // FD3D12BindlessDescriptorManager
 
-FD3D12BindlessDescriptorManager::FD3D12BindlessDescriptorManager(FD3D12Device* InDevice)
+FD3D12BindlessDescriptorManager::FD3D12BindlessDescriptorManager(FD3D12Device* InDevice, FD3D12BindlessDescriptorAllocator& InAllocator)
 	: FD3D12DeviceChild(InDevice)
+	, Allocator(InAllocator)
 {
 }
 
@@ -715,34 +719,17 @@ FD3D12BindlessDescriptorManager::~FD3D12BindlessDescriptorManager() = default;
 
 void FD3D12BindlessDescriptorManager::Init()
 {
-	FD3D12Adapter* Adapter = GetParentDevice()->GetParentAdapter();
-	ResourcesConfiguration = Adapter->GetBindlessManager().GetResourcesConfiguration();
-	SamplersConfiguration  = Adapter->GetBindlessManager().GetSamplersConfiguration();
+	ResourcesConfiguration = Allocator.GetResourcesConfiguration();
+	SamplersConfiguration  = Allocator.GetSamplersConfiguration();
 
 	if (ResourcesConfiguration != ERHIBindlessConfiguration::Disabled)
 	{
-		ResourceManager = MakeUnique<FD3D12BindlessResourceManager>(GetParentDevice(), Adapter);
+		ResourceManager = MakeUnique<FD3D12BindlessResourceManager>(GetParentDevice(), Allocator);
 	}
 
 	if (SamplersConfiguration != ERHIBindlessConfiguration::Disabled)
 	{
-		const TStatId Stats[] =
-		{
-			GET_STATID(STAT_SamplerDescriptorsAllocated),
-			GET_STATID(STAT_BindlessSamplerDescriptorsAllocated),
-		};
-
-		uint32 NumSamplerDescriptors = GBindlessSamplerDescriptorHeapSize;
-		if (NumSamplerDescriptors > D3D12_MAX_SHADER_VISIBLE_SAMPLER_HEAP_SIZE)
-		{
-			UE_LOG(LogD3D12RHI, Error, TEXT("D3D12.Bindless.SamplerDescriptorHeapSize was set to %d, which is higher than the D3D12 maximum of %d. Adjusting the value to prevent a crash."),
-				NumSamplerDescriptors,
-				D3D12_MAX_SHADER_VISIBLE_SAMPLER_HEAP_SIZE
-			);
-			NumSamplerDescriptors = D3D12_MAX_SHADER_VISIBLE_SAMPLER_HEAP_SIZE;
-		}
-
-		SamplerManager = MakeUnique<FD3D12BindlessSamplerManager>(GetParentDevice(), SamplersConfiguration, NumSamplerDescriptors, Stats);
+		SamplerManager = MakeUnique<FD3D12BindlessSamplerManager>(GetParentDevice(), Allocator);
 	}
 }
 
@@ -779,13 +766,13 @@ void FD3D12BindlessDescriptorManager::ImmediateFree(FRHIDescriptorHandle InHandl
 {
 	if (InHandle.GetType() == ERHIDescriptorHeapType::Standard && ResourceManager)
 	{
-		GetParentDevice()->GetParentAdapter()->GetBindlessManager().FreeResourceHandle(InHandle);
+		Allocator.FreeResourceHandle(InHandle);
 		return;
 	}
 
 	if (InHandle.GetType() == ERHIDescriptorHeapType::Sampler && SamplerManager)
 	{
-		GetParentDevice()->GetParentAdapter()->GetBindlessManager().FreeSamplerHandle(InHandle);
+		Allocator.FreeSamplerHandle(InHandle);
 		return;
 	}
 
@@ -801,6 +788,18 @@ void FD3D12BindlessDescriptorManager::DeferredFreeFromDestructor(FRHIDescriptorH
 	}
 }
 
+void FD3D12BindlessDescriptorManager::InitializeDescriptor(FRHIDescriptorHandle DstHandle, FD3D12SamplerState* SamplerState)
+{
+	if (SamplerManager)
+	{
+		SamplerManager->InitializeDescriptor(DstHandle, SamplerState);
+		return;
+	}
+
+	// Bad configuration?
+	checkNoEntry();
+}
+
 void FD3D12BindlessDescriptorManager::InitializeDescriptor(FRHIDescriptorHandle DstHandle, FD3D12View* View)
 {
 	if (DstHandle.GetType() == ERHIDescriptorHeapType::Standard && ResourceManager)
diff --git a/Engine/Source/Runtime/D3D12RHI/Private/D3D12BindlessDescriptors.h b/Engine/Source/Runtime/D3D12RHI/Private/D3D12BindlessDescriptors.h
--- a/Engine/Source/Runtime/D3D12RHI/Private/D3D12BindlessDescriptors.h
+++ b/Engine/Source/Runtime/D3D12RHI/Private/D3D12BindlessDescriptors.h
@@ -27,47 +27,45 @@ namespace UE::D3D12BindlessDescriptors
 }
 
 /** Manager for configuration settings and shared descriptor allocators, stored on the adapter. */
-class FD3D12BindlessManagerAdapter : public FD3D12AdapterChild
+class FD3D12BindlessDescriptorAllocator : public FD3D12AdapterChild
 {
 public:
-	void Init(FD3D12Adapter* InParent);
+	FD3D12BindlessDescriptorAllocator() = delete;
+	FD3D12BindlessDescriptorAllocator(FD3D12Adapter* InParent);
+
+	void Init();
 
 	ERHIBindlessConfiguration GetResourcesConfiguration() const { return BindlessResourcesConfiguration; }
 	ERHIBindlessConfiguration GetSamplersConfiguration() const { return BindlessSamplersConfiguration; }
 
-	bool GetResourcesAllowed() const { return bBindlessResourcesAllowed; }
-	bool GetSamplersAllowed() const { return bBindlessSamplersAllowed; }
+	bool AreResourcesBindless() const { return ResourceAllocator != nullptr; }
+	bool AreSamplersBindless() const { return SamplerAllocator != nullptr; }
 
 	// Bindless descriptor allocators are stored in the adapter, so descriptor handles can be allocated once and shared for multi-GPU objects
 	FRHIDescriptorHandle AllocateSamplerHandle();
 	FRHIDescriptorHandle AllocateResourceHandle();
+
 	void FreeSamplerHandle(FRHIDescriptorHandle Handle);
 	void FreeResourceHandle(FRHIDescriptorHandle Handle);
 
 	FCriticalSection& GetResourceHeapsCS() { return ResourceHeapsCS; }
 
 	uint32 GetResourceCapacity() const { return ResourceAllocator->GetCapacity(); }
+	uint32 GetSamplerCapacity()  const { return SamplerAllocator->GetCapacity(); }
 
 	bool GetResourceAllocatedRange(FRHIDescriptorAllocatorRange& OutAllocatedRange) const { return ResourceAllocator->GetAllocatedRange(OutAllocatedRange); }
 
-	FRHIDescriptorHandle& GetSamplerDefaultHandle() { return SamplerDefaultHandle; }
-
 #if D3D12RHI_CUSTOM_BINDLESS_RESOURCE_MANAGER
 	FRHIHeapDescriptorAllocator* GetResourceAllocator() { return ResourceAllocator; }
 #endif
 
 private:
-	bool bBindlessResourcesAllowed = false;
-	bool bBindlessSamplersAllowed = false;
-
 	ERHIBindlessConfiguration BindlessResourcesConfiguration{};
 	ERHIBindlessConfiguration BindlessSamplersConfiguration{};
 
 	FCriticalSection ResourceHeapsCS;
 	FRHIHeapDescriptorAllocator* ResourceAllocator = nullptr;
 	FRHIHeapDescriptorAllocator* SamplerAllocator = nullptr;
-
-	FRHIDescriptorHandle SamplerDefaultHandle;
 };
 
 /** Manager specifically for bindless sampler descriptors. */
@@ -75,10 +73,12 @@ class FD3D12BindlessSamplerManager : public FD3D12DeviceChild
 {
 public:
 	FD3D12BindlessSamplerManager() = delete;
-	FD3D12BindlessSamplerManager(FD3D12Device* InDevice, ERHIBindlessConfiguration InConfiguration, uint32 InNumDescriptors, TConstArrayView<TStatId> InStats);
+	FD3D12BindlessSamplerManager(FD3D12Device* InDevice, FD3D12BindlessDescriptorAllocator& InAllocator);
 
 	void CleanupResources();
 
+	void InitializeDescriptor(FRHIDescriptorHandle DstHandle, FD3D12SamplerState* SamplerState);
+
 	void OpenCommandList(FD3D12CommandContext& Context);
 	void CloseCommandList(FD3D12CommandContext& Context);
 
@@ -87,16 +87,15 @@ public:
 	FD3D12DescriptorHeap* GetHeap() const { return GpuHeap.GetReference(); }
 	ERHIBindlessConfiguration GetConfiguration() const { return Configuration; }
 
-	void InitializeSampler(FRHIDescriptorHandle DstHandle, FD3D12SamplerState* SamplerState);
-
 private:
 	FD3D12DescriptorHeapPtr      GpuHeap;
-	FRHIHeapDescriptorAllocator  Allocator;
 	ERHIBindlessConfiguration    Configuration;
 };
 
 #if !D3D12RHI_CUSTOM_BINDLESS_RESOURCE_MANAGER
 
+#define D3D12RHI_BINDLESS_RESOURCE_MANAGER_SUPPORTS_RESIZING 1
+
 // Helper container for all context related bindless state.
 struct FD3D12ContextBindlessState
 {
@@ -164,7 +163,7 @@ class FD3D12BindlessResourceManager : public FD3D12DeviceChild
 {
 public:
 	FD3D12BindlessResourceManager() = delete;
-	FD3D12BindlessResourceManager(FD3D12Device* InDevice, FD3D12Adapter* InAdapter);
+	FD3D12BindlessResourceManager(FD3D12Device* InDevice, FD3D12BindlessDescriptorAllocator& InAllocator);
 
 	void CleanupResources();
 	
@@ -238,12 +237,17 @@ struct FD3D12DescriptorHeapPair
 class FD3D12BindlessDescriptorManager : public FD3D12DeviceChild
 {
 public:
-	FD3D12BindlessDescriptorManager(FD3D12Device* InDevice);
+	FD3D12BindlessDescriptorManager(FD3D12Device* InDevice, FD3D12BindlessDescriptorAllocator& InAllocator);
 	~FD3D12BindlessDescriptorManager();
 
 	void Init();
 	void CleanupResources();
 
+	FD3D12BindlessDescriptorAllocator& GetAllocator() { return Allocator; }
+
+	FD3D12BindlessResourceManager* GetResourceManager() const { return ResourceManager.Get(); }
+	FD3D12BindlessSamplerManager*  GetSamplerManager()  const { return SamplerManager.Get();  }
+
 	ERHIBindlessConfiguration GetResourcesConfiguration() const { return ResourcesConfiguration; }
 	ERHIBindlessConfiguration GetSamplersConfiguration()  const { return SamplersConfiguration; }
 
@@ -262,8 +266,7 @@ public:
 	void GarbageCollect();
 	void Recycle(FD3D12DescriptorHeap* DescriptorHeap);
 
-	void InitializeSampler(FRHIDescriptorHandle DstHandle, FD3D12SamplerState* SamplerState) { SamplerManager->InitializeSampler(DstHandle, SamplerState); }
-
+	void InitializeDescriptor(FRHIDescriptorHandle DstHandle, FD3D12SamplerState* SamplerState);
 	void InitializeDescriptor(FRHIDescriptorHandle DstHandle, FD3D12View* View);
 	void UpdateDescriptor(FD3D12ContextArray const& Contexts, FRHIDescriptorHandle DstHandle, FD3D12View* SourceView);
 
@@ -282,9 +285,7 @@ public:
 #endif
 
 private:
-	// Needs direct access to underlying resource manager for resource allocation functions
-	friend class FD3D12BindlessManagerAdapter;
-	FD3D12BindlessResourceManager& GetResourceManager() const { return *ResourceManager; }
+	FD3D12BindlessDescriptorAllocator& Allocator;
 
 	TUniquePtr<FD3D12BindlessResourceManager> ResourceManager;
 	TUniquePtr<FD3D12BindlessSamplerManager>  SamplerManager;
@@ -293,20 +294,4 @@ private:
 	ERHIBindlessConfiguration SamplersConfiguration{};
 };
 
-#else // PLATFORM_SUPPORTS_BINDLESS_RENDERING
-
-// Stub version when bindless is compiled out
-class FD3D12BindlessManagerAdapter : public FD3D12AdapterChild
-{
-public:
-	inline void Init(FD3D12Adapter* InParent) { SetParentAdapter(InParent); }
-
-	inline FRHIDescriptorHandle AllocateSamplerHandle() { return FRHIDescriptorHandle(); }
-	inline FRHIDescriptorHandle AllocateResourceHandle() { return FRHIDescriptorHandle(); }
-	inline FRHIDescriptorHandle& GetSamplerDefaultHandle() { return SamplerDefaultHandle; }
-
-private:
-	FRHIDescriptorHandle SamplerDefaultHandle;
-};
-
 #endif // !PLATFORM_SUPPORTS_BINDLESS_RENDERING
diff --git a/Engine/Source/Runtime/D3D12RHI/Private/D3D12Device.cpp b/Engine/Source/Runtime/D3D12RHI/Private/D3D12Device.cpp
--- a/Engine/Source/Runtime/D3D12RHI/Private/D3D12Device.cpp
+++ b/Engine/Source/Runtime/D3D12RHI/Private/D3D12Device.cpp
@@ -102,7 +102,8 @@ FD3D12Device::FD3D12Device(FRHIGPUMask InGPUMask, FD3D12Adapter* InAdapter)
 	, ResidencyManager         (*this)
 	, DescriptorHeapManager    (this)
 #if PLATFORM_SUPPORTS_BINDLESS_RENDERING
-	, BindlessDescriptorManager(this)
+	, BindlessDescriptorAllocator(InAdapter->GetBindlessDescriptorAllocator())
+	, BindlessDescriptorManager(this, InAdapter->GetBindlessDescriptorAllocator())
 #endif
 	, GlobalSamplerHeap        (this)
 	, OnlineDescriptorManager  (this)
diff --git a/Engine/Source/Runtime/D3D12RHI/Private/D3D12Device.h b/Engine/Source/Runtime/D3D12RHI/Private/D3D12Device.h
--- a/Engine/Source/Runtime/D3D12RHI/Private/D3D12Device.h
+++ b/Engine/Source/Runtime/D3D12RHI/Private/D3D12Device.h
@@ -234,7 +234,8 @@ public:
 	// Descriptor Managers
 	inline FD3D12DescriptorHeapManager&     GetDescriptorHeapManager    () { return DescriptorHeapManager;     }
 #if PLATFORM_SUPPORTS_BINDLESS_RENDERING
-	inline FD3D12BindlessDescriptorManager& GetBindlessDescriptorManager() { return BindlessDescriptorManager; }
+	FD3D12BindlessDescriptorAllocator& GetBindlessDescriptorAllocator() { return BindlessDescriptorAllocator; }
+	FD3D12BindlessDescriptorManager& GetBindlessDescriptorManager() { return BindlessDescriptorManager; }
 #endif
 	inline FD3D12OnlineDescriptorManager&   GetOnlineDescriptorManager  () { return OnlineDescriptorManager;   }
 	inline FD3D12OfflineDescriptorManager&  GetOfflineDescriptorManager (ERHIDescriptorHeapType InType)
@@ -317,6 +318,7 @@ private:
 
 	FD3D12DescriptorHeapManager     DescriptorHeapManager;
 #if PLATFORM_SUPPORTS_BINDLESS_RENDERING
+	FD3D12BindlessDescriptorAllocator& BindlessDescriptorAllocator;
 	FD3D12BindlessDescriptorManager BindlessDescriptorManager;
 #endif
 	TArray<FD3D12OfflineDescriptorManager, TInlineAllocator<(uint32)ERHIDescriptorHeapType::Count>> OfflineDescriptorManagers;
diff --git a/Engine/Source/Runtime/D3D12RHI/Private/D3D12State.cpp b/Engine/Source/Runtime/D3D12RHI/Private/D3D12State.cpp
--- a/Engine/Source/Runtime/D3D12RHI/Private/D3D12State.cpp
+++ b/Engine/Source/Runtime/D3D12RHI/Private/D3D12State.cpp
@@ -800,12 +800,11 @@ FD3D12SamplerState::FD3D12SamplerState(FD3D12Device* InParent, const D3D12_SAMPL
 	GetParentDevice()->CreateSamplerInternal(Desc, OfflineDescriptor);
 
 #if PLATFORM_SUPPORTS_BINDLESS_RENDERING
-	FD3D12Adapter* Adapter = InParent->GetParentAdapter();
-	if (Adapter->GetBindlessManager().GetSamplersConfiguration() != ERHIBindlessConfiguration::Disabled)
+	if (InParent->GetBindlessDescriptorAllocator().GetSamplersConfiguration() != ERHIBindlessConfiguration::Disabled)
 	{
-		BindlessHandle = FirstLinkedObject ? FirstLinkedObject->BindlessHandle : InParent->GetParentAdapter()->GetBindlessManager().AllocateSamplerHandle();
+		BindlessHandle = FirstLinkedObject ? FirstLinkedObject->BindlessHandle : InParent->GetBindlessDescriptorAllocator().AllocateSamplerHandle();
 
-		InParent->GetBindlessDescriptorManager().InitializeSampler(BindlessHandle, this);
+		InParent->GetBindlessDescriptorManager().InitializeDescriptor(BindlessHandle, this);
 	}
 #endif
 }
diff --git a/Engine/Source/Runtime/D3D12RHI/Private/D3D12TextureReference.cpp b/Engine/Source/Runtime/D3D12RHI/Private/D3D12TextureReference.cpp
--- a/Engine/Source/Runtime/D3D12RHI/Private/D3D12TextureReference.cpp
+++ b/Engine/Source/Runtime/D3D12RHI/Private/D3D12TextureReference.cpp
@@ -5,7 +5,7 @@
 FD3D12RHITextureReference::FD3D12RHITextureReference(FD3D12Device* InDevice, FD3D12Texture* InReferencedTexture, FD3D12RHITextureReference* FirstLinkedObject)
 	: FD3D12DeviceChild(InDevice)
 #if PLATFORM_SUPPORTS_BINDLESS_RENDERING
-	, FRHITextureReference(InReferencedTexture, FirstLinkedObject ? FirstLinkedObject->BindlessHandle : InDevice->GetParentAdapter()->GetBindlessManager().AllocateResourceHandle())
+	, FRHITextureReference(InReferencedTexture, FirstLinkedObject ? FirstLinkedObject->BindlessHandle : InDevice->GetBindlessDescriptorAllocator().AllocateResourceHandle())
 #else
 	, FRHITextureReference(InReferencedTexture)
 #endif
diff --git a/Engine/Source/Runtime/D3D12RHI/Private/D3D12Util.cpp b/Engine/Source/Runtime/D3D12RHI/Private/D3D12Util.cpp
--- a/Engine/Source/Runtime/D3D12RHI/Private/D3D12Util.cpp
+++ b/Engine/Source/Runtime/D3D12RHI/Private/D3D12Util.cpp
@@ -1371,8 +1371,8 @@ const FD3D12RootSignature* FD3D12Adapter::GetGlobalRayTracingRootSignature(const
 	QBSS.bUseDiagnosticBuffer = true;
 
 #if PLATFORM_SUPPORTS_BINDLESS_RENDERING
-	QBSS.bUseDirectlyIndexedResourceHeap = BindlessManager.GetResourcesAllowed();
-	QBSS.bUseDirectlyIndexedSamplerHeap = BindlessManager.GetSamplersAllowed();
+	QBSS.bUseDirectlyIndexedResourceHeap = GetBindlessDescriptorAllocator().AreResourcesBindless();
+	QBSS.bUseDirectlyIndexedSamplerHeap = GetBindlessDescriptorAllocator().AreSamplersBindless();
 #endif
 
 	QBSSRegisterCounts.SamplerCount = MAX_SAMPLERS;
diff --git a/Engine/Source/Runtime/D3D12RHI/Private/D3D12View.cpp b/Engine/Source/Runtime/D3D12RHI/Private/D3D12View.cpp
--- a/Engine/Source/Runtime/D3D12RHI/Private/D3D12View.cpp
+++ b/Engine/Source/Runtime/D3D12RHI/Private/D3D12View.cpp
@@ -225,7 +225,7 @@ FD3D12View::FD3D12View(FD3D12Device* InDevice, ERHIDescriptorHeapType InHeapType
 	if (HeapType == ERHIDescriptorHeapType::Standard)
 	{
 		// Bindless handle is shared across linked objects, so pull the handle from the first object if available, or allocate if we are the first
-		BindlessHandle = FirstLinkedObject ? FirstLinkedObject->BindlessHandle : InDevice->GetParentAdapter()->GetBindlessManager().AllocateResourceHandle();
+		BindlessHandle = FirstLinkedObject ? FirstLinkedObject->BindlessHandle : InDevice->GetBindlessDescriptorAllocator().AllocateResourceHandle();
 	}
 #endif
 }
