Refactoring D3D12 Resource Collection implementation to not create a different buffer for each GPU now that descriptor indices are shared betwen GPUs.

diff --git a/Engine/Source/Runtime/D3D12RHI/Private/D3D12ResourceCollection.cpp b/Engine/Source/Runtime/D3D12RHI/Private/D3D12ResourceCollection.cpp
--- a/Engine/Source/Runtime/D3D12RHI/Private/D3D12ResourceCollection.cpp
+++ b/Engine/Source/Runtime/D3D12RHI/Private/D3D12ResourceCollection.cpp
@@ -15,9 +15,6 @@ FD3D12ResourceCollection::FD3D12ResourceCollection(FD3D12Device* InParent, FRHIC
 {
 	const uint32 GpuIndex = InParent->GetGPUIndex();
 
-	TArray<FRHIDescriptorHandle> Handles;
-	Handles.Reserve(InMembers.Num());
-
 	for (const FRHIResourceCollectionMember& Member : InMembers)
 	{
 		switch (Member.Type)
@@ -27,13 +24,11 @@ FD3D12ResourceCollection::FD3D12ResourceCollection(FD3D12Device* InParent, FRHIC
 			if (FRHITextureReference* TextureReferenceRHI = static_cast<FRHITexture*>(Member.Resource)->GetTextureReference())
 			{
 				FD3D12RHITextureReference* TextureReference = FD3D12CommandContext::RetrieveObject<FD3D12RHITextureReference>(TextureReferenceRHI, GpuIndex);
-				Handles.Emplace(TextureReference->GetDefaultBindlessHandle());
 				AllTextureReferences.Emplace(TextureReference);
 			}
 			else
 			{
 				FD3D12Texture* Texture = FD3D12CommandContext::RetrieveTexture(static_cast<FRHITexture*>(Member.Resource), GpuIndex);
-				Handles.Emplace(Texture->GetDefaultBindlessHandle());
 				AllSrvs.Emplace(Texture->GetShaderResourceView());
 			}
 		}
@@ -41,24 +36,18 @@ FD3D12ResourceCollection::FD3D12ResourceCollection(FD3D12Device* InParent, FRHIC
 		case FRHIResourceCollectionMember::EType::TextureReference:
 		{
 			FD3D12RHITextureReference* TextureReference = FD3D12CommandContext::RetrieveObject<FD3D12RHITextureReference>(Member.Resource, GpuIndex);
-			Handles.Emplace(TextureReference->GetDefaultBindlessHandle());
 			AllTextureReferences.Emplace(TextureReference);
 		}
 		break;
 		case FRHIResourceCollectionMember::EType::ShaderResourceView:
 		{
 			FD3D12ShaderResourceView_RHI* ShaderResourceView = FD3D12CommandContext::RetrieveObject<FD3D12ShaderResourceView_RHI>(Member.Resource, GpuIndex);
-			Handles.Emplace(ShaderResourceView->GetBindlessHandle());
 			AllSrvs.Emplace(ShaderResourceView);
 		}
 		break;
 		}
 	}
 
-	constexpr D3D12_RESOURCE_STATES States = D3D12_RESOURCE_STATE_GENERIC_READ;
-	const TArray<uint32> CollectionMemory = UE::RHICore::CreateResourceCollectionArray<FRHIDescriptorHandle>(Handles);
-	InBuffer->UploadResourceData(RHICmdList, FRHIGPUMask::FromIndex(GpuIndex), States, CollectionMemory.GetData(), CollectionMemory.GetTypeSize() * CollectionMemory.Num());
-
 	D3D12_SHADER_RESOURCE_VIEW_DESC SRVDesc{};
 	SRVDesc.Format = DXGI_FORMAT_R32_TYPELESS;
 	SRVDesc.ViewDimension = D3D12_SRV_DIMENSION_BUFFER;
@@ -78,18 +67,10 @@ FRHIDescriptorHandle FD3D12ResourceCollection::GetBindlessHandle() const
 	return BufferSRV->GetBindlessHandle();
 }
 
-static FD3D12Buffer* CreateCollectionBuffer(FD3D12DynamicRHI& RHI, FRHICommandListBase& RHICmdList, TConstArrayView<FRHIResourceCollectionMember> InMembers)
-{
-	const size_t BufferSize = UE::RHICore::CalculateResourceCollectionMemorySize(InMembers);
-	FRHIBufferDesc BufferDesc(BufferSize, 4, EBufferUsageFlags::Static | EBufferUsageFlags::ByteAddressBuffer);
-
-	FRHIResourceCreateInfo CreateInfo(TEXT("ResourceCollection"));
-	return RHI.CreateD3D12Buffer(&RHICmdList, BufferDesc, ERHIAccess::SRVMask, CreateInfo, nullptr, true);
-}
-
 FRHIResourceCollectionRef FD3D12DynamicRHI::RHICreateResourceCollection(FRHICommandListBase& RHICmdList, TConstArrayView<FRHIResourceCollectionMember> InMembers)
 {
-	FD3D12Buffer* Buffer = CreateCollectionBuffer(*this, RHICmdList, InMembers);
+	FRHIBuffer* RHIBuffer = UE::RHICore::CreateResourceCollectionBuffer(RHICmdList, InMembers);
+	FD3D12Buffer* Buffer = ResourceCast(RHIBuffer);
 
 	FRHIViewDesc::FBufferSRV::FInitializer ViewDesc = FRHIViewDesc::CreateBufferSRV().SetType(FRHIViewDesc::EBufferType::Raw);
 	FShaderResourceViewRHIRef ShaderResourceView = RHICmdList.CreateShaderResourceView(Buffer, ViewDesc);
diff --git a/Engine/Source/Runtime/RHICore/Internal/RHICoreResourceCollection.h b/Engine/Source/Runtime/RHICore/Internal/RHICoreResourceCollection.h
--- a/Engine/Source/Runtime/RHICore/Internal/RHICoreResourceCollection.h
+++ b/Engine/Source/Runtime/RHICore/Internal/RHICoreResourceCollection.h
@@ -90,18 +90,21 @@ namespace UE::RHICore
 		TArray<uint32> Memory;
 	};
 
+	inline FRHIBuffer* CreateResourceCollectionBuffer(FRHICommandListBase& RHICmdList, TConstArrayView<FRHIResourceCollectionMember> InMembers)
+	{
+		FResourceCollectionUpload UploadData(InMembers);
+
+		FRHIResourceCreateInfo CreateInfo(TEXT("ResourceCollection"), &UploadData);
+		return RHICmdList.CreateBuffer(UploadData.GetResourceDataSize(), EBufferUsageFlags::Static | EBufferUsageFlags::ByteAddressBuffer, sizeof(uint32), ERHIAccess::SRVMask, CreateInfo);
+	}
+
 	class FGenericResourceCollection : public FRHIResourceCollection
 	{
 	public:
 		FGenericResourceCollection(FRHICommandListBase& RHICmdList, TConstArrayView<FRHIResourceCollectionMember> InMembers)
 			: FRHIResourceCollection(InMembers)
+			, Buffer(CreateResourceCollectionBuffer(RHICmdList, InMembers))
 		{
-			FResourceCollectionUpload UploadData(InMembers);
-
-			const size_t BufferSize = UE::RHICore::CalculateResourceCollectionMemorySize(InMembers);
-			FRHIResourceCreateInfo CreateInfo(TEXT("ResourceCollection"), &UploadData);
-			Buffer = RHICmdList.CreateBuffer(BufferSize, EBufferUsageFlags::Static | EBufferUsageFlags::ByteAddressBuffer, sizeof(uint32), ERHIAccess::SRVMask, CreateInfo);
-
 			FRHIViewDesc::FBufferSRV::FInitializer ViewDesc = FRHIViewDesc::CreateBufferSRV().SetType(FRHIViewDesc::EBufferType::Raw);
 			ShaderResourceView = RHICmdList.CreateShaderResourceView(Buffer, ViewDesc);
 		}
