Path Tracer: Make sure that background RTPSO compilation triggers a refresh when complete

In cases where background RTPSO creation takes a long time, it can happen that the path tracer has finished accumulating empty samples and can miss the moment where the RTPSO first becomes valid.

Also tweak the logging of lengthy RTPSO creations to include background jobs if they take longer than 1 second.

#rnx

diff --git a/Engine/Source/Runtime/D3D12RHI/Private/D3D12RayTracing.cpp b/Engine/Source/Runtime/D3D12RHI/Private/D3D12RayTracing.cpp
--- a/Engine/Source/Runtime/D3D12RHI/Private/D3D12RayTracing.cpp
+++ b/Engine/Source/Runtime/D3D12RHI/Private/D3D12RayTracing.cpp
@@ -2289,9 +2289,10 @@ public:
 		ShaderStats.Sort([](const FShaderStats& A, const FShaderStats& B) { return B.CompileTimeMS < A.CompileTimeMS; });
 
 		const double TotalCreationTimeMS = 1000.0 * FPlatformTime::ToSeconds64(TotalCreationTime);
-		const float CreationTimeWarningThresholdMS = 10.0f;
-		const bool bAllowLogSlowCreation = !Initializer.bBackgroundCompilation; // Only report creation stalls on the critical path
-		if (bAllowLogSlowCreation && TotalCreationTimeMS > CreationTimeWarningThresholdMS)
+		// log a message if the RTPSO creating took a long time (with a higher threshold if this was a background compilation)
+		const double CreationTimeWarningThresholdMS = 10.0;
+		const double BackgroundCreatingTimeWarningThresholdMS = 1000.0;
+		if (TotalCreationTimeMS > (Initializer.bBackgroundCompilation ? CreationTimeWarningThresholdMS : BackgroundCreatingTimeWarningThresholdMS))
 		{
 			const double CompileTimeMS = 1000.0 * FPlatformTime::ToSeconds64(CompileTime);
 			const double LinkTimeMS = 1000.0 * FPlatformTime::ToSeconds64(LinkTime);
diff --git a/Engine/Source/Runtime/Renderer/Private/RayTracing/RayTracingMaterialHitShaders.cpp b/Engine/Source/Runtime/Renderer/Private/RayTracing/RayTracingMaterialHitShaders.cpp
--- a/Engine/Source/Runtime/Renderer/Private/RayTracing/RayTracingMaterialHitShaders.cpp
+++ b/Engine/Source/Runtime/Renderer/Private/RayTracing/RayTracingMaterialHitShaders.cpp
@@ -776,6 +776,13 @@ void FDeferredShadingSceneRenderer::CreateRayTracingMaterialPipeline(
 		PipelineState = FallbackPipelineState;
 	}
 
+	if (FallbackPipelineState != nullptr && PipelineState != FallbackPipelineState && bIsPathTracing && !View.bIsOfflineRender)
+	{
+		// When using path tracing, a change in pipeline state compared to the previous frame means some new materials got added to the RTPSO
+		// and we should restart sample accumulation.
+		Scene->InvalidatePathTracedOutput();
+	}
+
 	check(PipelineState);
 
 	const int32 OpaqueShadowMaterialIndex = FindRayTracingHitGroupIndex(PipelineState, OpaqueShadowShader, true);
