[31.40] Temp CSV stat for FastAllocatorNumPages

#rnx
#tests Tested that the stat is logged appropriately on PC with -dx12. Ensured the numbers look appropriate by graphing the same using csvtosvg.

diff --git a/Engine/Source/Runtime/D3D12RHI/Private/D3D12Allocation.cpp b/Engine/Source/Runtime/D3D12RHI/Private/D3D12Allocation.cpp
--- a/Engine/Source/Runtime/D3D12RHI/Private/D3D12Allocation.cpp
+++ b/Engine/Source/Runtime/D3D12RHI/Private/D3D12Allocation.cpp
@@ -13,6 +13,10 @@
 #include "HAL/LowLevelMemStats.h"
 #include "ProfilingDebugging/MemoryTrace.h"
 
+#if CSV_PROFILER_STATS
+CSV_DECLARE_CATEGORY_MODULE_EXTERN(CORE_API, FMemory);
+#endif
+
 // Fix for random GPU crashes on draw indirects on multiple IHVs. Force all indirect arg buffers as committed resources (see UE-115982)
 static int32 GD3D12AllowPoolAllocateIndirectArgBuffers = 1;
 static FAutoConsoleVariableRef CVarD3D12AllowPoolAllocateIndirectArgBuffers(
@@ -2250,6 +2254,10 @@ void FD3D12FastAllocatorPagePool::ReturnFastAllocatorPage(FD3D12FastAllocatorPag
 
 void FD3D12FastAllocatorPagePool::CleanupPages(uint64 FrameLag)
 {
+#if CSV_PROFILER_STATS
+	CSV_CUSTOM_STAT(FMemory, FastAllocatorNumPages, Pool.Num(), ECsvCustomStatOp::Max);
+#endif
+
 	if (Pool.Num() <= GD3D12FastAllocatorMinPagesToRetain)
 	{
 		return;
