Added padding to FD3D12ResourceDesc to make sure it's properly zeroed for hashing purposes.

#rnx

diff --git a/Engine/Source/Runtime/D3D12RHI/Private/D3D12Resources.h b/Engine/Source/Runtime/D3D12RHI/Private/D3D12Resources.h
--- a/Engine/Source/Runtime/D3D12RHI/Private/D3D12Resources.h
+++ b/Engine/Source/Runtime/D3D12RHI/Private/D3D12Resources.h
@@ -143,16 +143,26 @@ struct FD3D12ResourceDesc : public D3D12_RESOURCE_DESC
 	// PixelFormat for the Resource that aliases our current resource.
 	EPixelFormat UAVPixelFormat{ PF_Unknown };
 
-#if D3D12RHI_NEEDS_VENDOR_EXTENSIONS
-	bool bRequires64BitAtomicSupport : 1 = false;
-#endif
+	union
+	{
+		struct
+		{
+		#if D3D12RHI_NEEDS_VENDOR_EXTENSIONS
+			bool bRequires64BitAtomicSupport : 1;
+		#endif
 
-	bool bReservedResource : 1 = false;
+			bool bReservedResource : 1;
 
-	bool bBackBuffer : 1 = false;
+			bool bBackBuffer : 1;
 
-	// External resources are owned by another application or middleware, not the Engine
-	bool bExternal : 1 = false;
+			// External resources are owned by another application or middleware, not the Engine
+			bool bExternal : 1;
+		};
+		// All the bitfield value are set to zero by default (false)
+		uint8 PackedBitF = 0;
+	};
+	// We include this padding to make sure it's properly zeroed for hashing purposes. The element count should be adjusted if the struct grows (see the static_assert below)
+	uint8 Padding[5] = {};
 
 	// If we support the new format list casting, use the newer APIs; otherwise, fall back to our UAV Aliasing approach.
 #if D3D12RHI_SUPPORTS_UNCOMPRESSED_UAV
@@ -166,6 +176,8 @@ struct FD3D12ResourceDesc : public D3D12_RESOURCE_DESC
 #endif
 };
 
+static_assert(offsetof(FD3D12ResourceDesc, Padding) == 64 - 5);
+
 class FD3D12Resource : public FThreadSafeRefCountedObject, public FD3D12DeviceChild, public FD3D12MultiNodeGPUObject
 {
 private:
