Zero the FD3D12ResourceDesc struct via new helper base class to ensure padding is correctly zeroed for the purposes of hashing. This fixes a large memory leak due to uninitialized padding values causing an explosion of hash entries. This issue manifested with Clang; the previous code relied on a quirk of MSVC which caused the padding to be zeroed

diff --git a/Engine/Source/Runtime/D3D12RHI/Private/D3D12Resources.h b/Engine/Source/Runtime/D3D12RHI/Private/D3D12Resources.h
--- a/Engine/Source/Runtime/D3D12RHI/Private/D3D12Resources.h
+++ b/Engine/Source/Runtime/D3D12RHI/Private/D3D12Resources.h
@@ -124,17 +124,35 @@ private:
 	bool bRequiresResidencyTracking = bool(ENABLE_RESIDENCY_MANAGEMENT);
 };
 
-struct FD3D12ResourceDesc : public D3D12_RESOURCE_DESC
+
+// This can be used to guarantee that a struct's padding is zero, which is necessary for hashing in some cases
+// NOTE: classes with virtual methods are not supported. The struct must be the first declared base class
+// to ensure the correct ordering of the ctor calls
+template <typename T>
+struct TZeroedStruct
 {
-	FD3D12ResourceDesc() = default;
+	TZeroedStruct<T>()
+	{
+		FMemory::Memzero(this, sizeof(T));
+	}
+};
+
+struct FD3D12ResourceDesc : public TZeroedStruct<FD3D12ResourceDesc>, D3D12_RESOURCE_DESC
+{
+	FD3D12ResourceDesc() :
+		TZeroedStruct<FD3D12ResourceDesc>()
+		{
+		}
 	FD3D12ResourceDesc(const CD3DX12_RESOURCE_DESC& Other)
-		: D3D12_RESOURCE_DESC(Other)
+		: TZeroedStruct<FD3D12ResourceDesc>()
+		, D3D12_RESOURCE_DESC(Other)
 	{
 	}
 	
 	// TODO: use this type everywhere and disallow implicit conversion
 	/*explicit*/ FD3D12ResourceDesc(const D3D12_RESOURCE_DESC& Other)
-		: D3D12_RESOURCE_DESC(Other)
+		: TZeroedStruct<FD3D12ResourceDesc>()
+		, D3D12_RESOURCE_DESC(Other)
 	{
 	}
 
@@ -143,26 +161,16 @@ struct FD3D12ResourceDesc : public D3D12_RESOURCE_DESC
 	// PixelFormat for the Resource that aliases our current resource.
 	EPixelFormat UAVPixelFormat{ PF_Unknown };
 
-	union
-	{
-		struct
-		{
-		#if D3D12RHI_NEEDS_VENDOR_EXTENSIONS
-			bool bRequires64BitAtomicSupport : 1;
-		#endif
+#if D3D12RHI_NEEDS_VENDOR_EXTENSIONS
+	bool bRequires64BitAtomicSupport : 1 = false;
+#endif
 
-			bool bReservedResource : 1;
+	bool bReservedResource : 1 = false;
 
-			bool bBackBuffer : 1;
+	bool bBackBuffer : 1 = false;
 
-			// External resources are owned by another application or middleware, not the Engine
-			bool bExternal : 1;
-		};
-		// All the bitfield value are set to zero by default (false)
-		uint8 PackedBitF = 0;
-	};
-	// We include this padding to make sure it's properly zeroed for hashing purposes. The element count should be adjusted if the struct grows (see the static_assert below)
-	uint8 Padding[5] = {};
+	// External resources are owned by another application or middleware, not the Engine
+	bool bExternal : 1 = false;
 
 	// If we support the new format list casting, use the newer APIs; otherwise, fall back to our UAV Aliasing approach.
 #if D3D12RHI_SUPPORTS_UNCOMPRESSED_UAV
@@ -176,8 +184,6 @@ struct FD3D12ResourceDesc : public D3D12_RESOURCE_DESC
 #endif
 };
 
-static_assert(offsetof(FD3D12ResourceDesc, Padding) == 64 - 5);
-
 class FD3D12Resource : public FThreadSafeRefCountedObject, public FD3D12DeviceChild, public FD3D12MultiNodeGPUObject
 {
 private:
