Update VRS RHI to set combiners to PASSTHROUGH, PASSTHROUGH when rate is 1x1 and shading rate image is nullptr. Prevents validation errors when running on platforms that require specific shader compilation to support VRS.

diff --git a/Engine/Source/Runtime/D3D12RHI/Private/D3D12CommandContext.h b/Engine/Source/Runtime/D3D12RHI/Private/D3D12CommandContext.h
--- a/Engine/Source/Runtime/D3D12RHI/Private/D3D12CommandContext.h
+++ b/Engine/Source/Runtime/D3D12RHI/Private/D3D12CommandContext.h
@@ -623,8 +623,7 @@ public:
 		);
 
 	virtual void SetDepthBounds(float MinDepth, float MaxDepth);
-	virtual void SetShadingRate(EVRSShadingRate ShadingRate, const TStaticArray<EVRSRateCombiner, ED3D12VRSCombinerStages::Num>& Combiners);
-	virtual void SetShadingRateImage(FD3D12Resource* ShadingRateImage);
+	virtual void SetShadingRate(EVRSShadingRate ShadingRate, FD3D12Resource* ShadingRateImage, const TStaticArray<EVRSRateCombiner, ED3D12VRSCombinerStages::Num>& Combiners);
 
 	virtual void SetAsyncComputeBudgetInternal(EAsyncComputeBudget Budget) {}
 
diff --git a/Engine/Source/Runtime/D3D12RHI/Private/D3D12Commands.cpp b/Engine/Source/Runtime/D3D12RHI/Private/D3D12Commands.cpp
--- a/Engine/Source/Runtime/D3D12RHI/Private/D3D12Commands.cpp
+++ b/Engine/Source/Runtime/D3D12RHI/Private/D3D12Commands.cpp
@@ -825,7 +825,7 @@ void FD3D12CommandContext::RHISetGraphicsPipelineState(FRHIGraphicsPipelineState
 		}
 		else
 		{
-			// Also forces shading rate image attachment to be ignored
+			// This also forces shading rate image attachment to be ignored, so no need to set it to nullptr in the state cache
 			StateCache.SetShadingRate(EVRSShadingRate::VRSSR_1x1, VRSRB_Passthrough, VRSRB_Passthrough);
 		}
 	}
@@ -1959,42 +1959,37 @@ void FD3D12CommandContext::RHISetShadingRate(EVRSShadingRate ShadingRate, EVRSRa
 #endif
 }
 
-void FD3D12CommandContext::SetShadingRate(EVRSShadingRate ShadingRate, const TStaticArray<EVRSRateCombiner, ED3D12VRSCombinerStages::Num>& Combiners)
+void FD3D12CommandContext::SetShadingRate(EVRSShadingRate ShadingRate, FD3D12Resource* RateImageTexture, const TStaticArray<EVRSRateCombiner, ED3D12VRSCombinerStages::Num>& Combiners)
 {
- #if PLATFORM_SUPPORTS_VARIABLE_RATE_SHADING
- 	if (GRHISupportsPipelineVariableRateShading && GraphicsCommandList5())
- 	{
-		for (int32 CombinerIndex = 0; CombinerIndex < Combiners.Num(); ++CombinerIndex)
-		{
-			VRSCombiners[CombinerIndex] = ConvertShadingRateCombiner(Combiners[CombinerIndex]);
-		}
- 		VRSShadingRate = static_cast<D3D12_SHADING_RATE>(ShadingRate);
- 		GraphicsCommandList5()->RSSetShadingRate(VRSShadingRate, VRSCombiners);
- 	}
-	else
+#if	PLATFORM_SUPPORTS_VARIABLE_RATE_SHADING
+	if (GraphicsCommandList5())
 	{
-		// Ensure we're at a reasonable default in the case we're not supporting VRS.
-		for (int32 CombinerIndex = 0; CombinerIndex < Combiners.Num(); ++CombinerIndex)
+		if (GRHISupportsPipelineVariableRateShading)
 		{
-			VRSCombiners[CombinerIndex] = D3D12_SHADING_RATE_COMBINER_PASSTHROUGH;
-		}
-	}
- #endif
-}
+			if (ShadingRate == EVRSShadingRate::VRSSR_1x1 && RateImageTexture == nullptr)
+			{
+				// Make sure VRS is fully disabled when rate is 1x1 and no shading rate image is passed in
+				// Otherwise we may encounter validation issues on platforms where shaders must be compiled to support VRS
+				for (int32 CombinerIndex = 0; CombinerIndex < Combiners.Num(); ++CombinerIndex)
+				{
+					VRSCombiners[CombinerIndex] = D3D12_SHADING_RATE_COMBINER_PASSTHROUGH;
+				}
+			}
+			else
+			{
+				for (int32 CombinerIndex = 0; CombinerIndex < Combiners.Num(); ++CombinerIndex)
+				{
+					VRSCombiners[CombinerIndex] = ConvertShadingRateCombiner(Combiners[CombinerIndex]);
+				}
+			}
+			VRSShadingRate = static_cast<D3D12_SHADING_RATE>(ShadingRate);
+			GraphicsCommandList5()->RSSetShadingRate(VRSShadingRate, VRSCombiners);
 
-void FD3D12CommandContext::SetShadingRateImage(FD3D12Resource* RateImageTexture)
-{
-#if PLATFORM_SUPPORTS_VARIABLE_RATE_SHADING
-	if (GRHISupportsAttachmentVariableRateShading && GraphicsCommandList5())
-	{
-		if (RateImageTexture)
-		{
-			GraphicsCommandList5()->RSSetShadingRateImage(RateImageTexture->GetResource());
-		}
-		else
-		{
-			GraphicsCommandList5()->RSSetShadingRateImage(nullptr);
+			if (GRHISupportsAttachmentVariableRateShading) // In D3D12, support for attachment VRS implies support for pipeline VRS
+			{
+				GraphicsCommandList5()->RSSetShadingRateImage(RateImageTexture ? RateImageTexture->GetResource() : nullptr);
+			}
 		}
 	}
- #endif
+#endif
 }
diff --git a/Engine/Source/Runtime/D3D12RHI/Private/D3D12StateCache.cpp b/Engine/Source/Runtime/D3D12RHI/Private/D3D12StateCache.cpp
--- a/Engine/Source/Runtime/D3D12RHI/Private/D3D12StateCache.cpp
+++ b/Engine/Source/Runtime/D3D12RHI/Private/D3D12StateCache.cpp
@@ -541,16 +541,12 @@ void FD3D12StateCache::ApplyState(ERHIPipeline HardwarePipe, ED3D12PipelineType
 			CmdContext.SetDepthBounds(PipelineState.Graphics.MinDepth, PipelineState.Graphics.MaxDepth);
 		}
 
-		if (bNeedSetShadingRate)
+		// The shading rate image can influence RSSetShadingRate, because we want to set the combiners to passthrough/passthrough if the rate = 1x1 and the SRI is nullptr
+		if (bNeedSetShadingRate || bNeedSetShadingRateImage)
 		{
 			bNeedSetShadingRate = false;
-			CmdContext.SetShadingRate(PipelineState.Graphics.DrawShadingRate, PipelineState.Graphics.Combiners);
-		}
-
-		if (bNeedSetShadingRateImage)
-		{
 			bNeedSetShadingRateImage = false;
-			CmdContext.SetShadingRateImage(PipelineState.Graphics.ShadingRateImage);
+			CmdContext.SetShadingRate(PipelineState.Graphics.DrawShadingRate, PipelineState.Graphics.ShadingRateImage, PipelineState.Graphics.Combiners);
 		}
 	}
 
