Adding ED3D12ViewType to make sure only SRVs and UAVs allocate bindless descriptors.

diff --git a/Engine/Source/Runtime/D3D12RHI/Private/D3D12SRV.cpp b/Engine/Source/Runtime/D3D12RHI/Private/D3D12SRV.cpp
--- a/Engine/Source/Runtime/D3D12RHI/Private/D3D12SRV.cpp
+++ b/Engine/Source/Runtime/D3D12RHI/Private/D3D12SRV.cpp
@@ -10,13 +10,15 @@
 // -----------------------------------------------------------------------------------------------------
 
 FD3D12ShaderResourceView::FD3D12ShaderResourceView(FD3D12Device* InDevice, FD3D12ShaderResourceView* FirstLinkedObject)
-	: TD3D12View(InDevice, ERHIDescriptorHeapType::Standard, FirstLinkedObject)
-{}
+	: TD3D12View(InDevice, ED3D12ViewType::ShaderResource, ERHIDescriptorHeapType::Standard, FirstLinkedObject)
+{
+}
 
 FD3D12ShaderResourceView::FD3D12ShaderResourceView(FD3D12Device* InDevice, FD3D12ShaderResourceView* FirstLinkedObject, FD3D12RayTracingScene* InRayTracingScene)
-	: TD3D12View(InDevice, ERHIDescriptorHeapType::Standard, FirstLinkedObject)
+	: TD3D12View(InDevice, ED3D12ViewType::ShaderResource, ERHIDescriptorHeapType::Standard, FirstLinkedObject)
 	, RayTracingScene(InRayTracingScene)
-{}
+{
+}
 
 void FD3D12ShaderResourceView::UpdateResourceInfo(const FResourceInfo& InResource, const D3D12_SHADER_RESOURCE_VIEW_DESC& InD3DViewDesc, EFlags InFlags)
 {
diff --git a/Engine/Source/Runtime/D3D12RHI/Private/D3D12UAV.cpp b/Engine/Source/Runtime/D3D12RHI/Private/D3D12UAV.cpp
--- a/Engine/Source/Runtime/D3D12RHI/Private/D3D12UAV.cpp
+++ b/Engine/Source/Runtime/D3D12RHI/Private/D3D12UAV.cpp
@@ -10,8 +10,9 @@
 // -----------------------------------------------------------------------------------------------------
 
 FD3D12UnorderedAccessView::FD3D12UnorderedAccessView(FD3D12Device* InDevice, FD3D12UnorderedAccessView* FirstLinkedObject)
-	: TD3D12View(InDevice, ERHIDescriptorHeapType::Standard, FirstLinkedObject)
-{}
+	: TD3D12View(InDevice, ED3D12ViewType::UnorderedAccess, ERHIDescriptorHeapType::Standard, FirstLinkedObject)
+{
+}
 
 void FD3D12UnorderedAccessView::UpdateResourceInfo(const FResourceInfo& InResource, const D3D12_UNORDERED_ACCESS_VIEW_DESC& InD3DViewDesc, EFlags InFlags)
 {
diff --git a/Engine/Source/Runtime/D3D12RHI/Private/D3D12View.cpp b/Engine/Source/Runtime/D3D12RHI/Private/D3D12View.cpp
--- a/Engine/Source/Runtime/D3D12RHI/Private/D3D12View.cpp
+++ b/Engine/Source/Runtime/D3D12RHI/Private/D3D12View.cpp
@@ -213,7 +213,14 @@ FD3D12ViewRange::FD3D12ViewRange(D3D12_UNORDERED_ACCESS_VIEW_DESC const& ViewDes
 //
 // -----------------------------------------------------------------------------------------------------
 
-FD3D12View::FD3D12View(FD3D12Device* InDevice, ERHIDescriptorHeapType InHeapType, FD3D12View* FirstLinkedObject)
+#if PLATFORM_SUPPORTS_BINDLESS_RENDERING
+inline bool ViewTypeSupportsBindless(ED3D12ViewType InViewType)
+{
+	return (InViewType == ED3D12ViewType::ShaderResource || InViewType == ED3D12ViewType::UnorderedAccess);
+}
+#endif
+
+FD3D12View::FD3D12View(FD3D12Device* InDevice, ED3D12ViewType InViewType, ERHIDescriptorHeapType InHeapType, FD3D12View* FirstLinkedObject)
 	: FD3D12DeviceChild(InDevice)
 	, OfflineCpuHandle(InDevice->GetOfflineDescriptorManager(InHeapType).AllocateHeapSlot())
 #if PLATFORM_SUPPORTS_BINDLESS_RENDERING
@@ -222,7 +229,7 @@ FD3D12View::FD3D12View(FD3D12Device* InDevice, ERHIDescriptorHeapType InHeapType
 	, HeapType(InHeapType)
 {
 #if PLATFORM_SUPPORTS_BINDLESS_RENDERING
-	if (HeapType == ERHIDescriptorHeapType::Standard)
+	if (ViewTypeSupportsBindless(InViewType))
 	{
 		// Bindless handle is shared across linked objects, so pull the handle from the first object if available, or allocate if we are the first
 		BindlessHandle = FirstLinkedObject ? FirstLinkedObject->BindlessHandle : InDevice->GetBindlessDescriptorAllocator().AllocateResourceHandle();
@@ -332,8 +339,9 @@ void FD3D12View::UpdateView(FD3D12ContextArray const& Contexts, const FResourceI
 // -----------------------------------------------------------------------------------------------------
 
 FD3D12ConstantBufferView::FD3D12ConstantBufferView(FD3D12Device* InDevice, FD3D12ConstantBufferView* FirstLinkedObject)
-	: TD3D12View(InDevice, ERHIDescriptorHeapType::Standard, FirstLinkedObject)
-{}
+	: TD3D12View(InDevice, ED3D12ViewType::ConstantBuffer, ERHIDescriptorHeapType::Standard, FirstLinkedObject)
+{
+}
 
 void FD3D12ConstantBufferView::CreateView(FResourceInfo const& InResource, uint32 InOffset, uint32 InAlignedSize)
 {
@@ -368,8 +376,9 @@ void FD3D12ConstantBufferView::UpdateDescriptor()
 // -----------------------------------------------------------------------------------------------------
 
 FD3D12RenderTargetView::FD3D12RenderTargetView(FD3D12Device* InDevice, FD3D12RenderTargetView* FirstLinkedObject)
-	: TD3D12View(InDevice, ERHIDescriptorHeapType::RenderTarget, FirstLinkedObject)
-{}
+	: TD3D12View(InDevice, ED3D12ViewType::RenderTarget, ERHIDescriptorHeapType::RenderTarget, FirstLinkedObject)
+{
+}
 
 void FD3D12RenderTargetView::UpdateDescriptor()
 {
@@ -391,8 +400,9 @@ void FD3D12RenderTargetView::UpdateDescriptor()
 // -----------------------------------------------------------------------------------------------------
 
 FD3D12DepthStencilView::FD3D12DepthStencilView(FD3D12Device* InDevice, FD3D12DepthStencilView* FirstLinkedObject)
-	: TD3D12View(InDevice, ERHIDescriptorHeapType::DepthStencil, FirstLinkedObject)
-{}
+	: TD3D12View(InDevice, ED3D12ViewType::DepthStencil, ERHIDescriptorHeapType::DepthStencil, FirstLinkedObject)
+{
+}
 
 void FD3D12DepthStencilView::UpdateDescriptor()
 {
diff --git a/Engine/Source/Runtime/D3D12RHI/Private/D3D12View.h b/Engine/Source/Runtime/D3D12RHI/Private/D3D12View.h
--- a/Engine/Source/Runtime/D3D12RHI/Private/D3D12View.h
+++ b/Engine/Source/Runtime/D3D12RHI/Private/D3D12View.h
@@ -169,6 +169,16 @@ struct FD3D12ViewSubset
 	FD3D12ViewRange      Range;
 };
 
+// Enumeration of supported view types. Each view type must pass this up to their parent FD3D12View constructor.
+enum class ED3D12ViewType : uint8
+{
+	ShaderResource,
+	ConstantBuffer,
+	UnorderedAccess,
+	RenderTarget,
+	DepthStencil,
+};
+
 // Manages descriptor allocations and view creation
 class FD3D12View : public FD3D12DeviceChild, public FD3D12ShaderResourceRenameListener
 {
@@ -219,7 +229,8 @@ public:
 #endif
 
 protected:
-	FD3D12View(FD3D12Device* InDevice, ERHIDescriptorHeapType InHeapType, FD3D12View* FirstLinkedObject);
+	FD3D12View() = delete;
+	FD3D12View(FD3D12Device* InDevice, ED3D12ViewType InViewType, ERHIDescriptorHeapType InHeapType, FD3D12View* FirstLinkedObject);
 	virtual ~FD3D12View();
 		
 	virtual void UpdateDescriptor() = 0;
@@ -253,8 +264,8 @@ class TD3D12View : public FD3D12View
 protected:
 	TDesc D3DViewDesc;
 
-	TD3D12View(FD3D12Device* InDevice, ERHIDescriptorHeapType InHeapType, TD3D12View* FirstLinkedObject)
-		: FD3D12View(InDevice, InHeapType, FirstLinkedObject)
+	TD3D12View(FD3D12Device* InDevice, ED3D12ViewType InViewType, ERHIDescriptorHeapType InHeapType, TD3D12View* FirstLinkedObject)
+		: FD3D12View(InDevice, InViewType, InHeapType, FirstLinkedObject)
 	{}
 
 	void CreateView(FResourceInfo const& InResource, TDesc const& InD3DViewDesc)
