[Memory] Drop fast allocator retained pages from 5 to 1 to free up an additional of 4(pages) x 4(MB) = 16MB. Remove temp csv stat added to track page usage in CL 36341717.

#tests Graphed ratained pages across playtests and ReplayRuns to find they do not grow to > than 1 (hence safe to lower the value).
#rnx

diff --git a/Engine/Source/Runtime/D3D12RHI/Private/D3D12Allocation.cpp b/Engine/Source/Runtime/D3D12RHI/Private/D3D12Allocation.cpp
--- a/Engine/Source/Runtime/D3D12RHI/Private/D3D12Allocation.cpp
+++ b/Engine/Source/Runtime/D3D12RHI/Private/D3D12Allocation.cpp
@@ -13,10 +13,6 @@
 #include "HAL/LowLevelMemStats.h"
 #include "ProfilingDebugging/MemoryTrace.h"
 
-#if CSV_PROFILER_STATS
-CSV_DECLARE_CATEGORY_MODULE_EXTERN(CORE_API, FMemory);
-#endif
-
 // Fix for random GPU crashes on draw indirects on multiple IHVs. Force all indirect arg buffers as committed resources (see UE-115982)
 static int32 GD3D12AllowPoolAllocateIndirectArgBuffers = 1;
 static FAutoConsoleVariableRef CVarD3D12AllowPoolAllocateIndirectArgBuffers(
@@ -2254,10 +2250,6 @@ void FD3D12FastAllocatorPagePool::ReturnFastAllocatorPage(FD3D12FastAllocatorPag
 
 void FD3D12FastAllocatorPagePool::CleanupPages(uint64 FrameLag)
 {
-#if CSV_PROFILER_STATS
-	CSV_CUSTOM_STAT(FMemory, FastAllocatorNumPages, Pool.Num(), ECsvCustomStatOp::Max);
-#endif
-
 	if (Pool.Num() <= GD3D12FastAllocatorMinPagesToRetain)
 	{
 		return;
