Fix hangs when GPU crashes are detected by an API call wrapped in VERIFYD3D12RESULT, but GetDeviceRemovedReason returns S_OK.

Apparently this can happen in the wild, most likely because various things hook the swapchain and return DXGI_ERROR_DEVICE_REMOVED from Present, despite the device being fine. Now we explicitly set a flag when we see a device lost code anywhere and the interrupt thread triggers the GPU crash report mechanism if the flag is set, regardless of the status reported by GetDeviceRemovedReason.

#rnx

diff --git a/Engine/Source/Runtime/D3D12RHI/Private/D3D12Submission.cpp b/Engine/Source/Runtime/D3D12RHI/Private/D3D12Submission.cpp
--- a/Engine/Source/Runtime/D3D12RHI/Private/D3D12Submission.cpp
+++ b/Engine/Source/Runtime/D3D12RHI/Private/D3D12Submission.cpp
@@ -52,6 +52,8 @@ static FAutoConsoleVariableRef CVarD3D12SubmissionTimeout(
 	TEXT("The maximum time, in seconds, that a submitted GPU command list is allowed to take before the RHI reports a GPU hang"),
 	ECVF_RenderThreadSafe);
 
+static std::atomic<int> GGPUCrashDetected = false;
+
 class FD3D12Thread final : private FRunnable
 {
 public:
@@ -1167,6 +1169,9 @@ void FD3D12DynamicRHI::ProcessInterruptQueueUntil(FGraphEvent* GraphEvent)
 void FD3D12DynamicRHI::ProcessInterruptQueueOnGPUCrash()
 {
 	// This function will not return.
+
+	// We know this function was called due to a GPU crash, so let the thread know.
+	GGPUCrashDetected.store(true, std::memory_order_release);
 	
 	if (InterruptThread)
 	{
@@ -1425,6 +1430,13 @@ FD3D12DynamicRHI::FProcessResult FD3D12DynamicRHI::ProcessInterruptQueue()
 		CheckForDeviceRemoved(CurrentQueue);
 	});
 
+	if (GGPUCrashDetected.load(std::memory_order_relaxed))
+	{
+		// If this was set by ProcessInterruptQueueOnGPUCrash, we know a crash was detected, so process it immediately. We can't always rely on
+		// queue processing to catch it, as GetDeviceRemovedReason sometimes returns S_OK despite an earlier API call having reported a lost device.
+		TerminateOnGPUCrash();
+	}
+
 	return Result;
 }
 
