Add support for bindles to work graph shader bundles.

diff --git a/Engine/Shaders/Private/ShaderBundleWorkGraphDispatch.usf b/Engine/Shaders/Private/ShaderBundleWorkGraphDispatch.usf
--- a/Engine/Shaders/Private/ShaderBundleWorkGraphDispatch.usf
+++ b/Engine/Shaders/Private/ShaderBundleWorkGraphDispatch.usf
@@ -10,7 +10,7 @@ struct FEntryNodeRecord
 {
 	uint DispatchGridSize : SV_DispatchGrid;
 	uint RecordCount;
-	uint2 PlatformData;
+	uint3 PlatformData;
 };
 
 /* Entry node which reads required group count and launches the shading nodes. */
@@ -41,7 +41,14 @@ void WorkGraphMainCS(
 	if (bValidItem)
 	{
 		const uint ArgOffset  = ArgBaseOffset + DispatchIndex * ArgStride;
+
+#if ENABLE_BINDLESS_RESOURCES && PLATFORM_SUPPORTS_BINDLESS
+		const uint RecordArgBufferHandle = InputRecord.Get().PlatformData.z;
+		ByteAddressBuffer BindlessRecordArgBuffer = GetResourceFromHeap(ByteAddressBuffer, RecordArgBufferHandle);
+		DispatchGridSize = BindlessRecordArgBuffer.Load(ArgOffset);
+#else
 		DispatchGridSize = RecordArgBuffer.Load(ArgOffset);
+#endif
 	}
 
 	const bool bValidDispatch = DispatchGridSize > 0;
diff --git a/Engine/Source/Runtime/D3D12RHI/Private/D3D12WorkGraph.cpp b/Engine/Source/Runtime/D3D12RHI/Private/D3D12WorkGraph.cpp
--- a/Engine/Source/Runtime/D3D12RHI/Private/D3D12WorkGraph.cpp
+++ b/Engine/Source/Runtime/D3D12RHI/Private/D3D12WorkGraph.cpp
@@ -8,6 +8,7 @@
 #include "D3D12RHIPrivate.h"
 #include "D3D12ResourceCollection.h"
 #include "D3D12Shader.h"
+#include "D3D12TextureReference.h"
 #include "PipelineStateCache.h"
 #include "ShaderBundles.h"
 
@@ -209,10 +210,9 @@ struct FWorkGraphShaderBundleBinder
 	FD3D12CommandContext& Context;
 	FShaderBundleBinderOps& BinderOps;
 	const uint32 GPUIndex;
-#if PLATFORM_SUPPORTS_BINDLESS_RENDERING
+	const EShaderFrequency Frequency;
 	const bool bBindlessResources;
 	const bool bBindlessSamplers;
-#endif
 
 	uint32 CBVVersions[MAX_CBS];
 	uint32 SRVVersions[MAX_SRVS];
@@ -229,28 +229,15 @@ struct FWorkGraphShaderBundleBinder
 	D3D12_CPU_DESCRIPTOR_HANDLE LocalUAVs[MAX_UAVS];
 	D3D12_CPU_DESCRIPTOR_HANDLE LocalSamplers[MAX_SAMPLERS];
 
-	FWorkGraphShaderBundleBinder(FD3D12CommandContext& InContext, FShaderBundleBinderOps& InBinderOps, FD3D12ShaderData const* ShaderData)
+	FWorkGraphShaderBundleBinder(FD3D12CommandContext& InContext, EShaderFrequency InFrequency, FShaderBundleBinderOps& InBinderOps, FD3D12ShaderData const* ShaderData)
 		: Context(InContext)
 		, BinderOps(InBinderOps)
 		, GPUIndex(InContext.GetGPUIndex())
-#if PLATFORM_SUPPORTS_BINDLESS_RENDERING
-		, bBindlessResources(EnumHasAnyFlags(ShaderData->ResourceCounts.UsageFlags, EShaderResourceUsageFlags::BindlessResources))
-		, bBindlessSamplers(EnumHasAnyFlags(ShaderData->ResourceCounts.UsageFlags, EShaderResourceUsageFlags::BindlessSamplers))
-#endif
-	{
-	}
-
-#if PLATFORM_SUPPORTS_BINDLESS_RENDERING
-	void SetBindlessHandle(const FRHIDescriptorHandle& Handle, uint32 Offset)
+		, Frequency(InFrequency)
+		, bBindlessResources(ShaderData->UsesBindlessResources())
+		, bBindlessSamplers(ShaderData->UsesBindlessSamplers())
 	{
-		if (Handle.IsValid())
-		{
-			checkNoEntry();
-			//const uint32 BindlessIndex = Handle.GetIndex();
-			//ConstantBuffer.UpdateConstant(reinterpret_cast<const uint8*>(&BindlessIndex), Offset, 4);
-		}
 	}
-#endif
 
 	void SetUAV(FRHIUnorderedAccessView* InUnorderedAccessView, uint32 Index, bool bClearResources = false)
 	{
@@ -262,21 +249,15 @@ struct FWorkGraphShaderBundleBinder
 			BinderOps.AddResourceClear(UAV);
 		}
 
-#if PLATFORM_SUPPORTS_BINDLESS_RENDERING
-		if (bBindlessResources)
-		{
-			checkNoEntry();
-			//Context.StateCache.QueueBindlessUAV(Frequency, D3D12UnorderedAccessView);
-		}
-		else
-#endif
+		if (!bBindlessResources)
 		{
 			FD3D12OfflineDescriptor Descriptor = UAV->GetOfflineCpuHandle();
 			LocalUAVs[Index] = Descriptor;
 			UAVVersions[Index] = Descriptor.GetVersion();
 			BoundUAVMask |= 1ull << Index;
-			BinderOps.AddResourceTransition(UAV);
 		}
+
+		BinderOps.AddResourceTransition(UAV);
 	}
 
 	void SetSRV(FRHIShaderResourceView* InShaderResourceView, uint32 Index)
@@ -284,21 +265,15 @@ struct FWorkGraphShaderBundleBinder
 		FD3D12ShaderResourceView_RHI* SRV = FD3D12CommandContext::RetrieveObject<FD3D12ShaderResourceView_RHI>(InShaderResourceView, GPUIndex);
 		check(SRV != nullptr);
 
-#if PLATFORM_SUPPORTS_BINDLESS_RENDERING
-		if (bBindlessResources)
-		{
-			checkNoEntry();
-			//Context.StateCache.QueueBindlessSRV(Frequency, SRV);
-		}
-		else
-#endif
+		if (!bBindlessResources)
 		{
 			FD3D12OfflineDescriptor Descriptor = SRV->GetOfflineCpuHandle();
 			LocalSRVs[Index] = Descriptor;
 			SRVVersions[Index] = Descriptor.GetVersion();
 			BoundSRVMask |= 1ull << Index;
-			BinderOps.AddResourceTransition(SRV);
 		}
+		
+		BinderOps.AddResourceTransition(SRV);
 	}
 
 	void SetTexture(FRHITexture* InTexture, uint32 Index)
@@ -306,21 +281,15 @@ struct FWorkGraphShaderBundleBinder
 		FD3D12ShaderResourceView* SRV = FD3D12CommandContext::RetrieveTexture(InTexture, GPUIndex)->GetShaderResourceView();
 		check(SRV != nullptr);
 
-#if PLATFORM_SUPPORTS_BINDLESS_RENDERING
-		if (bBindlessResources)
-		{
-			checkNoEntry();
-			//Context.StateCache.QueueBindlessSRV(Frequency, D3D12ShaderResourceView);
-		}
-		else
-#endif
+		if (!bBindlessResources)
 		{
 			FD3D12OfflineDescriptor Descriptor = SRV->GetOfflineCpuHandle();
 			LocalSRVs[Index] = Descriptor;
 			SRVVersions[Index] = Descriptor.GetVersion();
 			BoundSRVMask |= 1ull << Index;
-			BinderOps.AddResourceTransition(SRV);
 		}
+
+		BinderOps.AddResourceTransition(SRV);
 	}
 
 	void SetSampler(FRHISamplerState* InSampler, uint32 Index)
@@ -328,13 +297,7 @@ struct FWorkGraphShaderBundleBinder
 		FD3D12SamplerState* Sampler = FD3D12CommandContext::RetrieveObject<FD3D12SamplerState>(InSampler, GPUIndex);
 		check(Sampler != nullptr);
 
-#if PLATFORM_SUPPORTS_BINDLESS_RENDERING
-		if (bBindlessSamplers)
-		{
-			// Nothing to do, only needs constants set
-		}
-		else
-#endif
+		if (!bBindlessSamplers)
 		{
 			FD3D12OfflineDescriptor Descriptor = Sampler->OfflineDescriptor;
 			LocalSamplers[Index] = Descriptor;
@@ -345,27 +308,39 @@ struct FWorkGraphShaderBundleBinder
 
 	void SetResourceCollection(FRHIResourceCollection* ResourceCollection, uint32 Index)
 	{
-#if PLATFORM_SUPPORTS_BINDLESS_RENDERING
 		FD3D12ResourceCollection* D3D12ResourceCollection = FD3D12CommandContext::RetrieveObject<FD3D12ResourceCollection>(ResourceCollection, GPUIndex);
 		FD3D12ShaderResourceView* SRV = D3D12ResourceCollection ? D3D12ResourceCollection->GetShaderResourceView() : nullptr;
 
+		check(bBindlessResources);
 		if (bBindlessResources)
 		{
-			checkNoEntry();
-			//Context.StateCache.QueueBindlessSRV(Frequency, SRV);
+			Context.StateCache.QueueBindlessSRV(Frequency, SRV);
 		}
-#endif // PLATFORM_SUPPORTS_BINDLESS_RENDERING
 	}
 };
 
+/** Wrapper for constant buffer and it's underlying resource allocation. */
+struct FAllocatedConstantBuffer
+{
+	FAllocatedConstantBuffer(FD3D12CommandContext& Context)
+		: ResourceLocation(Context.GetParentDevice())
+	{
+	}
+
+	FD3D12ConstantBuffer* ConstantBuffer = nullptr;
+	FD3D12ResourceLocation ResourceLocation;
+};
+
 // Record bindings from shader bundle parameters.
 static void RecordBindings(
-	FD3D12CommandContext* Context,
+	FD3D12CommandContext& Context,
+	EShaderFrequency Frequency,
 	FD3D12ExplicitDescriptorCache& TransientDescriptorCache,
 	FShaderBundleBinderOps& BinderOps,
 	uint32 WorkerIndex,
 	FRHIWorkGraphShader* WorkGraphShaderRHI,
 	FRHIBatchedShaderParameters const& Parameters,
+	FAllocatedConstantBuffer const& SharedConstantBuffer,
 	FUint32Vector4 const& Constants,
 	TArrayView<uint32> RootArgs
 )
@@ -379,9 +354,9 @@ static void RecordBindings(
 	// With shader root constants, we should never hit this expensive path!
 	// If we hit this, check if the shaders in the bundle had loose
 	// uniform parameters added to it recently, falling into this path.
-	check(!D3D12WorkGraphShader->UsesGlobalUniformBuffer());
+	check(!D3D12WorkGraphShader->UsesGlobalUniformBuffer() || SharedConstantBuffer.ConstantBuffer != nullptr);
 
-	FWorkGraphShaderBundleBinder BundleBinder(*Context, BinderOps, D3D12WorkGraphShader);
+	FWorkGraphShaderBundleBinder BundleBinder(Context, Frequency, BinderOps, D3D12WorkGraphShader);
 
 	FD3D12UniformBuffer* BundleUniformBuffers[MAX_CBS] = { nullptr };
 
@@ -389,38 +364,6 @@ static void RecordBindings(
 
 	const bool bClearUAVResources = false;
 
-#if PLATFORM_SUPPORTS_BINDLESS_RENDERING
-	for (const FRHIShaderParameterResource& Parameter : Parameters.BindlessParameters)
-	{
-		if (FRHIResource* Resource = Parameter.Resource)
-		{
-			FRHIDescriptorHandle Handle;
-
-			switch (Parameter.Type)
-			{
-			case FRHIShaderParameterResource::EType::Texture:
-				Handle = static_cast<FRHITexture*>(Resource)->GetDefaultBindlessHandle();
-				BundleBinder.SetTexture(static_cast<FRHITexture*>(Resource), Parameter.Index);
-				break;
-			case FRHIShaderParameterResource::EType::ResourceView:
-				Handle = static_cast<FRHIShaderResourceView*>(Resource)->GetBindlessHandle();
-				BundleBinder.SetSRV(static_cast<FRHIShaderResourceView*>(Resource), Parameter.Index);
-				break;
-			case FRHIShaderParameterResource::EType::UnorderedAccessView:
-				Handle = static_cast<FRHIUnorderedAccessView*>(Resource)->GetBindlessHandle();
-				BundleBinder.SetUAV(static_cast<FRHIUnorderedAccessView*>(Resource), Parameter.Index, bClearUAVResources);
-				break;
-			case FRHIShaderParameterResource::EType::Sampler:
-				Handle = static_cast<FRHISamplerState*>(Resource)->GetBindlessHandle();
-				BundleBinder.SetSampler(static_cast<FRHISamplerState*>(Resource), Parameter.Index);
-				break;
-			}
-
-			BundleBinder.SetBindlessHandle(Handle, Parameter.Index);
-		}
-	}
-#endif
-
 	for (const FRHIShaderParameterResource& Parameter : Parameters.ResourceParameters)
 	{
 		switch (Parameter.Type)
@@ -450,20 +393,26 @@ static void RecordBindings(
 		}
 	}
 
-	UE::RHICore::ApplyStaticUniformBuffers(WorkGraphShaderRHI, Context->GetStaticUniformBuffers(),
+	UE::RHICore::ApplyStaticUniformBuffers(WorkGraphShaderRHI, Context.GetStaticUniformBuffers(),
 		[&](int32 BufferIndex, FRHIUniformBuffer* Buffer)
 		{
-			BundleUniformBuffers[BufferIndex] = Context->RetrieveObject<FD3D12UniformBuffer>(Buffer);
+			BundleUniformBuffers[BufferIndex] = Context.RetrieveObject<FD3D12UniformBuffer>(Buffer);
 		});
 
 
 	uint32 FakeDirtyUniformBuffers = ~(0u);
 	UE::RHICore::SetResourcesFromTables(BundleBinder, *WorkGraphShaderRHI, FakeDirtyUniformBuffers, BundleUniformBuffers
 #if ENABLE_RHI_VALIDATION
-		, Context->Tracker
+		, Context.Tracker
 #endif
 	);
 
+	if (SharedConstantBuffer.ConstantBuffer != nullptr)
+	{
+		check(BundleUniformBuffers[0] == nullptr);
+		BundleBinder.BoundCBVMask |= 1ull << 0;
+	}
+
 	for (uint32 CBVIndex = 0; CBVIndex < MAX_CBS; ++CBVIndex)
 	{
 		FD3D12UniformBuffer* UniformBuffer = BundleUniformBuffers[CBVIndex];
@@ -525,6 +474,15 @@ static void RecordBindings(
 		FMemory::Memcpy(&RootArgs[BindSlotOffset], &ResourceDescriptorTableBaseGPU, sizeof(ResourceDescriptorTableBaseGPU));
 	}
 
+	if (SharedConstantBuffer.ConstantBuffer != nullptr)
+	{
+		const uint32 BindSlot = D3D12WorkGraphShader->RootSignature->CBVRDBindSlot(SF_Compute, 0);
+		const uint32 BindSlotOffset = D3D12WorkGraphShader->RootSignature->GetBindSlotOffsetInBytes(BindSlot) / 4;
+
+		D3D12_GPU_VIRTUAL_ADDRESS Address = SharedConstantBuffer.ResourceLocation.GetGPUVirtualAddress();
+		FMemory::Memcpy(&RootArgs[BindSlotOffset], &Address, sizeof(Address));
+	}
+
 	for (uint32 CBVIndex = 0; CBVIndex < NumCBVs; ++CBVIndex)
 	{
 		const uint32 BindSlot = D3D12WorkGraphShader->RootSignature->CBVRDBindSlot(SF_Compute, CBVIndex);
@@ -550,9 +508,137 @@ static void RecordBindings(
 	}
 }
 
+struct FD3D12BindlessConstantSetter
+{
+	FD3D12CommandContext& Context;
+	FD3D12ConstantBuffer& ConstantBuffer;
+	const uint32 GpuIndex;
+	const EShaderFrequency Frequency;
+
+	FD3D12BindlessConstantSetter(FD3D12CommandContext& InContext, EShaderFrequency InFrequency)
+		: Context(InContext)
+		, ConstantBuffer(InContext.StageConstantBuffers[SF_Compute])
+		, GpuIndex(InContext.GetGPUIndex())
+		, Frequency(InFrequency)
+	{
+	}
+
+	void SetBindlessHandle(const FRHIDescriptorHandle& Handle, uint32 Offset)
+	{
+		if (Handle.IsValid())
+		{
+			const uint32 BindlessIndex = Handle.GetIndex();
+			ConstantBuffer.UpdateConstant(reinterpret_cast<const uint8*>(&BindlessIndex), Offset, 4);
+		}
+	}
+
+	void SetUAV(FD3D12UnorderedAccessView* D3D12UnorderedAccessView, uint32 Offset)
+	{
+		SetBindlessHandle(D3D12UnorderedAccessView->GetBindlessHandle(), Offset);
+		Context.StateCache.QueueBindlessUAV(Frequency, D3D12UnorderedAccessView);
+	}
+
+	void SetUAV(FRHIUnorderedAccessView* InUnorderedAccessView, uint32 Offset)
+	{
+		FD3D12UnorderedAccessView_RHI* D3D12UnorderedAccessView = FD3D12CommandContext::RetrieveObject<FD3D12UnorderedAccessView_RHI>(InUnorderedAccessView, GpuIndex);
+		SetUAV(static_cast<FD3D12UnorderedAccessView*>(D3D12UnorderedAccessView), Offset);
+	}
+
+	void SetSRV(FD3D12ShaderResourceView* D3D12ShaderResourceView, uint32 Offset)
+	{
+		SetBindlessHandle(D3D12ShaderResourceView->GetBindlessHandle(), Offset);
+		Context.StateCache.QueueBindlessSRV(Frequency, D3D12ShaderResourceView);
+	}
+
+	void SetSRV(FRHIShaderResourceView* InShaderResourceView, uint32 Offset)
+	{
+		FD3D12ShaderResourceView_RHI* D3D12ShaderResourceView = FD3D12CommandContext::RetrieveObject<FD3D12ShaderResourceView_RHI>(InShaderResourceView, GpuIndex);
+		SetSRV(static_cast<FD3D12ShaderResourceView*>(D3D12ShaderResourceView), Offset);
+	}
+
+	void SetTexture(FRHITexture* InTexture, uint32 Offset)
+	{
+		FD3D12Texture* D3D12Texture = FD3D12CommandContext::RetrieveTexture(InTexture, GpuIndex);
+		FD3D12ShaderResourceView* D3D12ShaderResourceView = D3D12Texture ? D3D12Texture->GetShaderResourceView() : nullptr;
+
+		SetBindlessHandle(InTexture->GetDefaultBindlessHandle(), Offset);
+		Context.StateCache.QueueBindlessSRV(Frequency, D3D12ShaderResourceView);
+	}
+
+	void SetSampler(FRHISamplerState* InSampler, uint32 Offset)
+	{
+		FD3D12SamplerState* D3D12SamplerState = FD3D12CommandContext::RetrieveObject<FD3D12SamplerState>(InSampler, GpuIndex);
+
+		SetBindlessHandle(D3D12SamplerState->GetBindlessHandle(), Offset);
+	}
+
+	void SetResourceCollection(FRHIResourceCollection* ResourceCollection, uint32 Offset)
+	{
+		if (FD3D12ResourceCollection* D3D12ResourceCollection = FD3D12CommandContext::RetrieveObject<FD3D12ResourceCollection>(ResourceCollection, GpuIndex))
+		{
+			FD3D12ShaderResourceView* D3D12ShaderResourceView = D3D12ResourceCollection->GetShaderResourceView();
+			Context.StateCache.QueueBindlessSRV(Frequency, D3D12ShaderResourceView);
+			Context.StateCache.QueueBindlessSRVs(Frequency, D3D12ResourceCollection->AllSrvs);
+
+			// We have to go through each TextureReference to get the most recent version.
+			for (FD3D12RHITextureReference* TextureReference : D3D12ResourceCollection->AllTextureReferences)
+			{
+				if (FD3D12Texture* Texture = GetD3D12TextureFromRHITexture(TextureReference))
+				{
+					Context.StateCache.QueueBindlessSRV(Frequency, Texture->GetShaderResourceView());
+				}
+			}
+		}
+	}
+
+	[[nodiscard]] void Finalize(FAllocatedConstantBuffer& OutConstantBuffer)
+	{
+		OutConstantBuffer.ConstantBuffer = &ConstantBuffer;
+		ConstantBuffer.Version(OutConstantBuffer.ResourceLocation, false);
+	}
+};
+
+void SetShaderBundleSharedBindlessConstants(FD3D12CommandContext& Context, TConstArrayView<FRHIShaderParameterResource> SharedBindlessParameters, FAllocatedConstantBuffer& OutConstantBuffer)
+{
+	if (SharedBindlessParameters.Num())
+	{
+		FD3D12BindlessConstantSetter Setter(Context, SF_Compute);
+
+		for (const FRHIShaderParameterResource& Parameter : SharedBindlessParameters)
+		{
+			if (FRHIResource* Resource = Parameter.Resource)
+			{
+				switch (Parameter.Type)
+				{
+				case FRHIShaderParameterResource::EType::Texture:
+					Setter.SetTexture(static_cast<FRHITexture*>(Resource), Parameter.Index);
+					break;
+				case FRHIShaderParameterResource::EType::ResourceView:
+					Setter.SetSRV(static_cast<FRHIShaderResourceView*>(Resource), Parameter.Index);
+					break;
+				case FRHIShaderParameterResource::EType::UnorderedAccessView:
+					Setter.SetUAV(static_cast<FRHIUnorderedAccessView*>(Resource), Parameter.Index);
+					break;
+				case FRHIShaderParameterResource::EType::Sampler:
+					break;
+				case FRHIShaderParameterResource::EType::ResourceCollection:
+					Setter.SetResourceCollection(static_cast<FRHIResourceCollection*>(Resource), Parameter.Index);
+					break;
+				}
+			}
+		}
+
+		Setter.Finalize(OutConstantBuffer);
+	}
+}
+
 #endif // D3D12_RHI_WORKGRAPHS
 
-void FD3D12CommandContext::DispatchWorkGraphShaderBundle(FRHIShaderBundle* ShaderBundle, FRHIBuffer* RecordArgBuffer, TConstArrayView<FRHIShaderParameterResource> SharedBindlessParameters, TConstArrayView<FRHIShaderBundleComputeDispatch> Dispatches)
+void FD3D12CommandContext::DispatchWorkGraphShaderBundle(
+	FRHIShaderBundle* ShaderBundle, 
+	FRHIBuffer* RecordArgBuffer, 
+	TConstArrayView<FRHIShaderParameterResource> SharedBindlessParameters, 
+	TConstArrayView<FRHIShaderBundleComputeDispatch> Dispatches)
 {
 #if D3D12_RHI_WORKGRAPHS
 
@@ -562,6 +648,7 @@ void FD3D12CommandContext::DispatchWorkGraphShaderBundle(FRHIShaderBundle* Shade
 
 	TShaderRef<FDispatchShaderBundleWorkGraph> WorkGraphGlobalShader = GetGlobalShaderMap(GMaxRHIFeatureLevel)->GetShader<FDispatchShaderBundleWorkGraph>();
 	FD3D12WorkGraphShader* WorkGraphGlobalShaderRHI = static_cast<FD3D12WorkGraphShader*>(WorkGraphGlobalShader.GetWorkGraphShader());
+	const bool bBindlessResources = WorkGraphGlobalShaderRHI->UsesBindlessResources();
 
 	uint32 ViewDescriptorCount = WorkGraphGlobalShaderRHI->ResourceCounts.NumSRVs + WorkGraphGlobalShaderRHI->ResourceCounts.NumCBs + WorkGraphGlobalShaderRHI->ResourceCounts.NumUAVs;
 	uint32 SamplerDescriptorCount = WorkGraphGlobalShaderRHI->ResourceCounts.NumSamplers;
@@ -628,7 +715,10 @@ void FD3D12CommandContext::DispatchWorkGraphShaderBundle(FRHIShaderBundle* Shade
 	uint32 MinRootArgBufferSizeInDWords = (Pipeline->RootArgStrideInBytes / 4) * (Pipeline->MaxRootArgOffset + 1);
 	LocalRootArgs.AddZeroed(MinRootArgBufferSizeInDWords);
 
-	auto RecordTask = [this, &LocalRootArgs, Pipeline, &TransientDescriptorCache, &ValidRecords, &Dispatches, &BinderOps](FTaskContext& Context, int32 RecordIndex)
+	FAllocatedConstantBuffer SharedConstantBuffer(*this);
+	SetShaderBundleSharedBindlessConstants(*this, SharedBindlessParameters, SharedConstantBuffer);
+
+	auto RecordTask = [this, &LocalRootArgs, Pipeline, &TransientDescriptorCache, &ValidRecords, &Dispatches, &BinderOps, &SharedConstantBuffer](FTaskContext& Context, int32 RecordIndex)
 	{
 		uint32 DispatchIndex = ValidRecords[RecordIndex];
 		const FRHIShaderBundleComputeDispatch& Dispatch = Dispatches[DispatchIndex];
@@ -638,12 +728,14 @@ void FD3D12CommandContext::DispatchWorkGraphShaderBundle(FRHIShaderBundle* Shade
 		check((Pipeline->RootArgStrideInBytes / 4) * (RootArgOffset + 1) <= (uint32)LocalRootArgs.Num());
 
 		RecordBindings(
-			this,
+			*this,
+			SF_Compute,
 			TransientDescriptorCache,
 			BinderOps[Context.WorkerIndex],
 			Context.WorkerIndex,
 			Dispatch.WorkGraphShader,
 			*Dispatch.Parameters,
+			SharedConstantBuffer,
 			Dispatch.Constants,
 			MakeArrayView(&LocalRootArgs[RootArgOffset * Pipeline->RootArgStrideInBytes / 4], Pipeline->RootArgStrideInBytes / 4)
 		);
@@ -716,42 +808,55 @@ void FD3D12CommandContext::DispatchWorkGraphShaderBundle(FRHIShaderBundle* Shade
 	FlushResourceBarriers();
 
 	// Apply the transient descriptor heaps.
+	// Note this only uses transient heap for non-bindless.
 	SetExplicitDescriptorCache(TransientDescriptorCache);
 
-	TSharedPtr<FD3D12ShaderResourceView> RecordArgBufferSRV;
-
 	// Gather root arguments for global work graph.
-	int32 DispatchSRVTable = INDEX_NONE;
-	{
-		D3D12_CPU_DESCRIPTOR_HANDLE LocalSRVs[MAX_SRVS];
+	FD3D12Buffer* RecordArgBufferPtr = FD3D12DynamicRHI::ResourceCast(RecordArgBuffer);
 
-		FD3D12Buffer* RecordArgBufferPtr = FD3D12DynamicRHI::ResourceCast(RecordArgBuffer);
+	D3D12_SHADER_RESOURCE_VIEW_DESC SRVDesc = {};
+	SRVDesc.ViewDimension = D3D12_SRV_DIMENSION_BUFFER;
+	SRVDesc.Format = DXGI_FORMAT_R32_TYPELESS;
+	SRVDesc.Shader4ComponentMapping = D3D12_DEFAULT_SHADER_4_COMPONENT_MAPPING;
+	SRVDesc.Buffer.FirstElement = 0;
+	SRVDesc.Buffer.NumElements = RecordArgBufferPtr->GetSize() >> 2u;
+	SRVDesc.Buffer.StructureByteStride = 0;
+	SRVDesc.Buffer.Flags = D3D12_BUFFER_SRV_FLAG_RAW;
 
-		D3D12_SHADER_RESOURCE_VIEW_DESC SRVDesc = {};
-		SRVDesc.ViewDimension = D3D12_SRV_DIMENSION_BUFFER;
-		SRVDesc.Format = DXGI_FORMAT_R32_TYPELESS;
-		SRVDesc.Shader4ComponentMapping = D3D12_DEFAULT_SHADER_4_COMPONENT_MAPPING;
-		SRVDesc.Buffer.FirstElement = 0;
-		SRVDesc.Buffer.NumElements = RecordArgBufferPtr->GetSize() >> 2u;
-		SRVDesc.Buffer.StructureByteStride = 0;
-		SRVDesc.Buffer.Flags = D3D12_BUFFER_SRV_FLAG_RAW;
+	TSharedPtr<FD3D12ShaderResourceView> RecordArgBufferSRV;
+	RecordArgBufferSRV = MakeShared<FD3D12ShaderResourceView>(GetParentDevice(), nullptr); // Always single GPU object, so FirstLinkedObject is nullptr
+	RecordArgBufferSRV->CreateView(RecordArgBufferPtr, SRVDesc, FD3D12ShaderResourceView::EFlags::None);
 
-		RecordArgBufferSRV = MakeShared<FD3D12ShaderResourceView>(GetParentDevice(), nullptr);			// Always single GPU object, so FirstLinkedObject is nullptr
-		RecordArgBufferSRV->CreateView(RecordArgBufferPtr, SRVDesc, FD3D12ShaderResourceView::EFlags::None);
+	uint32 RecordArgBufferBindlessHandle = 0;
+	
+	if (bBindlessResources)
+	{
+		RecordArgBufferBindlessHandle = RecordArgBufferSRV->GetBindlessHandle().GetIndex();
+		check(RecordArgBufferBindlessHandle != INDEX_NONE);
 
-		LocalSRVs[WorkGraphGlobalShader->RecordArgBufferParam.GetBaseIndex()] = RecordArgBufferSRV->GetOfflineCpuHandle();
-		DispatchSRVTable = TransientDescriptorCache.Allocate(LocalSRVs, 1, D3D12_DESCRIPTOR_HEAP_TYPE_CBV_SRV_UAV, 0);
+		StateCache.ApplyBindlessResources(nullptr, SF_Compute, SF_NumStandardFrequencies);
+		
+		GraphicsCommandList()->SetComputeRootSignature(WorkGraphGlobalShaderRHI->RootSignature->GetRootSignature());
 	}
-	check(DispatchSRVTable != INDEX_NONE);
+	else
+	{
+		int32 DispatchSRVTable = INDEX_NONE;
+		{
+			D3D12_CPU_DESCRIPTOR_HANDLE LocalSRVs[MAX_SRVS];
+			LocalSRVs[WorkGraphGlobalShader->RecordArgBufferParam.GetBaseIndex()] = RecordArgBufferSRV->GetOfflineCpuHandle();
+			DispatchSRVTable = TransientDescriptorCache.Allocate(LocalSRVs, 1, D3D12_DESCRIPTOR_HEAP_TYPE_CBV_SRV_UAV, 0);
+		}
+		check(DispatchSRVTable != INDEX_NONE);
 
-	const D3D12_GPU_DESCRIPTOR_HANDLE DispatchSRVHandle = TransientDescriptorCache.ViewHeap.GetDescriptorGPU(DispatchSRVTable);
-	const uint32 SRVBindSlot = WorkGraphGlobalShaderRHI->RootSignature->SRVRDTBindSlot(SF_Compute);
-	check(SRVBindSlot != 0xFF);
+		const D3D12_GPU_DESCRIPTOR_HANDLE DispatchSRVHandle = TransientDescriptorCache.ViewHeap.GetDescriptorGPU(DispatchSRVTable);
+		const uint32 SRVBindSlot = WorkGraphGlobalShaderRHI->RootSignature->SRVRDTBindSlot(SF_Compute);
+		check(SRVBindSlot != 0xFF);
 
+		GraphicsCommandList()->SetComputeRootSignature(WorkGraphGlobalShaderRHI->RootSignature->GetRootSignature());
+		GraphicsCommandList()->SetComputeRootDescriptorTable(SRVBindSlot, DispatchSRVHandle);
+	}
+	
 	// Kick off the work graph	
-	GraphicsCommandList()->SetComputeRootSignature(WorkGraphGlobalShaderRHI->RootSignature->GetRootSignature());
-	GraphicsCommandList()->SetComputeRootDescriptorTable(SRVBindSlot, DispatchSRVHandle);
-
 	D3D12_SET_PROGRAM_DESC SetProgramDesc = {};
 	SetProgramDesc.Type = D3D12_PROGRAM_TYPE_WORK_GRAPH;
 	SetProgramDesc.WorkGraph.ProgramIdentifier = Pipeline->ProgramIdentifier;
@@ -760,7 +865,7 @@ void FD3D12CommandContext::DispatchWorkGraphShaderBundle(FRHIShaderBundle* Shade
 	SetProgramDesc.WorkGraph.NodeLocalRootArgumentsTable = NodeLocalRootArgumentsTable;
 	GraphicsCommandList10()->SetProgram(&SetProgramDesc);
 
-	FDispatchShaderBundleWorkGraph::FEntryNodeRecord InputRecord = FDispatchShaderBundleWorkGraph::MakeInputRecord(NumRecords, ShaderBundle->ArgOffset, ShaderBundle->ArgStride);
+	FDispatchShaderBundleWorkGraph::FEntryNodeRecord InputRecord = FDispatchShaderBundleWorkGraph::MakeInputRecord(NumRecords, ShaderBundle->ArgOffset, ShaderBundle->ArgStride, RecordArgBufferBindlessHandle);
 
 	if (!GShaderBundleSkipDispatch)
 	{
@@ -776,7 +881,7 @@ void FD3D12CommandContext::DispatchWorkGraphShaderBundle(FRHIShaderBundle* Shade
 	// Pipeline state memory should now be initialized.
 	Pipeline->bInitialized = true;
 
-	// Restore old global descriptor heaps
+	// Restore global descriptor heaps if necessary.
 	UnsetExplicitDescriptorCache();
 
 	// We did not write through the state cache, so we need to invalidate it so subsequent workloads correctly re-bind state
diff --git a/Engine/Source/Runtime/RenderCore/Private/ShaderBundles.cpp b/Engine/Source/Runtime/RenderCore/Private/ShaderBundles.cpp
--- a/Engine/Source/Runtime/RenderCore/Private/ShaderBundles.cpp
+++ b/Engine/Source/Runtime/RenderCore/Private/ShaderBundles.cpp
@@ -45,8 +45,6 @@ bool FDispatchShaderBundleWorkGraph::ShouldCompilePermutation(const FGlobalShade
 
 void FDispatchShaderBundleWorkGraph::ModifyCompilationEnvironment(const FGlobalShaderPermutationParameters& Parameters, FShaderCompilerEnvironment& OutEnvironment)
 {
-	OutEnvironment.CompilerFlags.Add(CFLAG_ForceBindful);
-
 	OutEnvironment.SetDefine(TEXT("THREADGROUP_SIZEX"), ThreadGroupSizeX);
 	OutEnvironment.SetDefine(TEXT("MAX_DISPATCHGRID_SIZEX"), (CVarShaderBundleMaxSize.GetValueOnAnyThread() + ThreadGroupSizeX - 1) / ThreadGroupSizeX);
 	FGlobalShader::ModifyCompilationEnvironment(Parameters, OutEnvironment);
@@ -58,7 +56,7 @@ int32 FDispatchShaderBundleWorkGraph::GetMaxShaderBundleSize()
 	return ShaderBundleMaxSize; 
 }
 
-FDispatchShaderBundleWorkGraph::FEntryNodeRecord FDispatchShaderBundleWorkGraph::MakeInputRecord(uint32 RecordCount, uint32 ArgOffset, uint32 ArgStride)
+FDispatchShaderBundleWorkGraph::FEntryNodeRecord FDispatchShaderBundleWorkGraph::MakeInputRecord(uint32 RecordCount, uint32 ArgOffset, uint32 ArgStride, uint32 ArgsBindlessHandle)
 {
-	return FEntryNodeRecord{ (RecordCount + ThreadGroupSizeX - 1) / ThreadGroupSizeX, RecordCount, FUintPoint(ArgOffset, ArgStride) };
+	return FEntryNodeRecord{ (RecordCount + ThreadGroupSizeX - 1) / ThreadGroupSizeX, RecordCount, FUintVector(ArgOffset, ArgStride, ArgsBindlessHandle) };
 }
diff --git a/Engine/Source/Runtime/RenderCore/Public/ShaderBundles.h b/Engine/Source/Runtime/RenderCore/Public/ShaderBundles.h
--- a/Engine/Source/Runtime/RenderCore/Public/ShaderBundles.h
+++ b/Engine/Source/Runtime/RenderCore/Public/ShaderBundles.h
@@ -49,7 +49,7 @@ class FDispatchShaderBundleWorkGraph : public FGlobalShader
 	FDispatchShaderBundleWorkGraph(const ShaderMetaType::CompiledShaderInitializerType& Initializer)
 		: FGlobalShader(Initializer)
 	{
-		RecordArgBufferParam.Bind(Initializer.ParameterMap, TEXT("RecordArgBuffer"), SPF_Mandatory);
+		RecordArgBufferParam.Bind(Initializer.ParameterMap, TEXT("RecordArgBuffer"), SPF_Optional);
 	}
 
 	static const uint32 ThreadGroupSizeX = 64;
@@ -66,7 +66,7 @@ class FDispatchShaderBundleWorkGraph : public FGlobalShader
 	{
 		uint32 DispatchGridSize;
 		uint32 RecordCount;
-		FUintPoint PlatformData;
+		FUintVector PlatformData;
 	};
-	RENDERCORE_API static FEntryNodeRecord MakeInputRecord(uint32 RecordCount, uint32 ArgOffset, uint32 ArgStride);
+	RENDERCORE_API static FEntryNodeRecord MakeInputRecord(uint32 RecordCount, uint32 ArgOffset, uint32 ArgStride, uint32 ArgsBindlessHandle);
 };
