- Make sure element count for hitrecord bindless SRV is computed correctly and we are not loosing entries due to divide by 4 for RAW buffer view
- Always align the buffer allocation size and offset to multiple of 16 so raw views with complete data can be created

diff --git a/Engine/Source/Runtime/D3D12RHI/Private/D3D12Buffer.cpp b/Engine/Source/Runtime/D3D12RHI/Private/D3D12Buffer.cpp
--- a/Engine/Source/Runtime/D3D12RHI/Private/D3D12Buffer.cpp
+++ b/Engine/Source/Runtime/D3D12RHI/Private/D3D12Buffer.cpp
@@ -134,6 +134,8 @@ void FD3D12Buffer::UploadResourceData(FRHICommandListBase& RHICmdList, FResource
 
 FD3D12SyncPointRef FD3D12Buffer::UploadResourceDataViaCopyQueue(FResourceArrayUploadInterface* InResourceArray)
 {
+	TRACE_CPUPROFILER_EVENT_SCOPE(UploadResourceDataViaCopyQueue);
+
 	// assume not dynamic and not on async thread (probably fine but untested)
 	check(IsInRHIThread() || IsInRenderingThread());
 	check(!(GetUsage() & BUF_AnyDynamic));
@@ -144,7 +146,10 @@ FD3D12SyncPointRef FD3D12Buffer::UploadResourceDataViaCopyQueue(FResourceArrayUp
 	FD3D12ResourceLocation SrcResourceLoc(GetParentDevice());
 	void* pData = GetParentDevice()->GetDefaultFastAllocator().Allocate(BufferSize, 4UL, &SrcResourceLoc);
 	check(pData);
-	FMemory::Memcpy(pData, InResourceArray->GetResourceData(), BufferSize);
+	{
+		TRACE_CPUPROFILER_EVENT_SCOPE(CopyToUploadMemory);
+		FMemory::Memcpy(pData, InResourceArray->GetResourceData(), BufferSize);
+	}
 
 	// Allocate copy queue command list and perform the copy op
 	FD3D12Device* Device = SrcResourceLoc.GetParentDevice();
@@ -198,7 +203,7 @@ void FD3D12Adapter::AllocateBuffer(FD3D12Device* Device,
 		check(InResourceStateMode != ED3D12ResourceStateMode::MultiState);
 		check(InCreateState == D3D12_RESOURCE_STATE_GENERIC_READ);
 		GetUploadHeapAllocator(Device->GetGPUIndex()).AllocUploadResource(Size, Alignment, ResourceLocation);
-		check(ResourceLocation.GetSize() == Size);
+		check(ResourceLocation.GetSize() >= Size);
 	}
 	else
 	{
@@ -211,7 +216,7 @@ void FD3D12Adapter::AllocateBuffer(FD3D12Device* Device,
 			Device->GetDefaultBufferAllocator().AllocDefaultResource(D3D12_HEAP_TYPE_DEFAULT, InDesc, InUsage, InResourceStateMode, InCreateState, ResourceLocation, Alignment, InDebugName);
 		}
 		ResourceLocation.SetOwner(Buffer);
-		check(ResourceLocation.GetSize() == Size);
+		check(ResourceLocation.GetSize() >= Size);
 	}
 }
 
@@ -233,7 +238,7 @@ FD3D12Buffer* FD3D12Adapter::CreateRHIBuffer(
 
 	UE_TRACE_METADATA_SCOPE_ASSET_FNAME(FName(InDebugName), GetRHIBufferClassName(ClassName), OwnerName);
 
-	check(InDesc.Width == BufferDesc.Size);
+	check(InDesc.Width >= BufferDesc.Size);
 
 	FD3D12Buffer* BufferOut = nullptr;
 
@@ -366,6 +371,9 @@ void FD3D12Buffer::GetResourceDescAndAlignment(uint64 InSize, uint32 InStride, E
 {
 	ResourceDesc = CD3DX12_RESOURCE_DESC::Buffer(InSize);
 
+	// Align size to 16 so RAW buffer view can be created without loosing and data at the end when dividing num elements by 4
+	ResourceDesc.Width = Align(ResourceDesc.Width, 16);
+
 	if (EnumHasAnyFlags(InUsage, BUF_UnorderedAccess))
 	{
 		ResourceDesc.Flags |= D3D12_RESOURCE_FLAG_ALLOW_UNORDERED_ACCESS;
@@ -397,7 +405,7 @@ void FD3D12Buffer::GetResourceDescAndAlignment(uint64 InSize, uint32 InStride, E
 	else
 	{
 		// Structured buffers, non-ByteAddress buffers, need to be aligned to their stride to ensure that they can be addressed correctly with element based offsets.
-		Alignment = (InStride > 0) && (EnumHasAnyFlags(InUsage, BUF_StructuredBuffer) || !EnumHasAnyFlags(InUsage, BUF_ByteAddressBuffer | BUF_DrawIndirect)) ? InStride : 16;
+		Alignment = (InStride > 0) && (EnumHasAnyFlags(InUsage, BUF_StructuredBuffer) || !EnumHasAnyFlags(InUsage, BUF_ByteAddressBuffer | BUF_DrawIndirect)) ? FMath::LeastCommonMultiplier(InStride, 16) : 16;
 	}
 }
 
diff --git a/Engine/Source/Runtime/D3D12RHI/Private/D3D12RayTracing.cpp b/Engine/Source/Runtime/D3D12RHI/Private/D3D12RayTracing.cpp
--- a/Engine/Source/Runtime/D3D12RHI/Private/D3D12RayTracing.cpp
+++ b/Engine/Source/Runtime/D3D12RHI/Private/D3D12RayTracing.cpp
@@ -3067,7 +3067,7 @@ void FD3D12RayTracingGeometry::AllocateBufferSRVs(uint32 InGPUIndex)
 		SRVDesc.Format = DXGI_FORMAT_R32_TYPELESS;
 		SRVDesc.Shader4ComponentMapping = D3D12_DEFAULT_SHADER_4_COMPONENT_MAPPING;
 		SRVDesc.Buffer.FirstElement = (Initializer.IndexBufferOffset + IndexBuffer->ResourceLocation.GetOffsetFromBaseOfResource()) >> 2u;
-		SRVDesc.Buffer.NumElements = FMath::Max((uint32)1, (Initializer.TotalPrimitiveCount * 3 * IndexBuffer->GetStride()) >> 2u);
+		SRVDesc.Buffer.NumElements = FMath::Max((uint32)1, ((Initializer.TotalPrimitiveCount * 3 * IndexBuffer->GetStride()) + 3) >> 2u);
 		SRVDesc.Buffer.Flags = D3D12_BUFFER_SRV_FLAG_RAW;
 		SRVDesc.Buffer.StructureByteStride = 0;
 
@@ -3088,11 +3088,11 @@ void FD3D12RayTracingGeometry::AllocateBufferSRVs(uint32 InGPUIndex)
 		SRVDesc.Buffer.FirstElement = (Segment.VertexBufferOffset + VertexBuffer->ResourceLocation.GetOffsetFromBaseOfResource()) >> 2u;
 		if (Initializer.GeometryType == RTGT_Procedural)
 		{
-			SRVDesc.Buffer.NumElements = Segment.NumPrimitives * Segment.VertexBufferStride / 4; //< NumElements in R32 size
+			SRVDesc.Buffer.NumElements = ((Segment.NumPrimitives * Segment.VertexBufferStride) + 3) / 4; //< NumElements in R32 size
 		}
 		else
 		{
-			SRVDesc.Buffer.NumElements = FMath::Max((uint32)1, Segment.MaxVertices * Segment.VertexBufferStride / 4); //< NumElements in R32 size
+			SRVDesc.Buffer.NumElements = FMath::Max((uint32)1, ((Segment.MaxVertices * Segment.VertexBufferStride) + 3) / 4); //< NumElements in R32 size
 		}
 		SRVDesc.Buffer.Flags = D3D12_BUFFER_SRV_FLAG_RAW;
 		SRVDesc.Buffer.StructureByteStride = 0;
