Delete unused code from FD3D12RayTracingGeometry constructor.

diff --git a/Engine/Source/Runtime/D3D12RHI/Private/D3D12RayTracing.cpp b/Engine/Source/Runtime/D3D12RHI/Private/D3D12RayTracing.cpp
--- a/Engine/Source/Runtime/D3D12RHI/Private/D3D12RayTracing.cpp
+++ b/Engine/Source/Runtime/D3D12RHI/Private/D3D12RayTracing.cpp
@@ -2826,22 +2826,11 @@ FD3D12RayTracingGeometry::FD3D12RayTracingGeometry(FRHICommandListBase& RHICmdLi
 
 	checkf(Initializer.Segments.Num() > 0, TEXT("Ray tracing geometry must be initialized with at least one segment."));
 
-	// #yuriy_todo: get flags directly through the initializer
-	ERayTracingAccelerationStructureFlags BuildFlags = GetRayTracingAccelerationStructureBuildFlags(Initializer);
-
 	GeometryDescs.SetNumUninitialized(Initializer.Segments.Num());
 	TranslateRayTracingGeometryDescs(Initializer, GeometryDescs);
 
 	SetDirty(FRHIGPUMask::All(), true);
 
-	D3D12_BUILD_RAYTRACING_ACCELERATION_STRUCTURE_INPUTS PrebuildDescInputs = {};
-
-	PrebuildDescInputs.Type = D3D12_RAYTRACING_ACCELERATION_STRUCTURE_TYPE_BOTTOM_LEVEL;
-	PrebuildDescInputs.DescsLayout = D3D12_ELEMENTS_LAYOUT_ARRAY;
-	PrebuildDescInputs.NumDescs = GeometryDescs.Num();
-	PrebuildDescInputs.pGeometryDescs = GeometryDescs.GetData();
-	PrebuildDescInputs.Flags = TranslateRayTracingAccelerationStructureFlags(BuildFlags);
-
 	struct FOfflineBVHHeader
 	{
 		uint64 ResultSize;
