[RayTracing] Define shared FD3D12RayTracingOfflineBvhHeader struct in D3D12RayTracingResources.h

diff --git a/Engine/Source/Runtime/D3D12RHI/Internal/D3D12RayTracingResources.h b/Engine/Source/Runtime/D3D12RHI/Internal/D3D12RayTracingResources.h
--- a/Engine/Source/Runtime/D3D12RHI/Internal/D3D12RayTracingResources.h
+++ b/Engine/Source/Runtime/D3D12RHI/Internal/D3D12RayTracingResources.h
@@ -1,7 +1,7 @@
 // Copyright Epic Games, Inc. All Rights Reserved.
 
 /*=============================================================================
-	D3D12ShaderResources.h: Internal D3D12 RHI RayTracing definitions.
+	D3D12RayTracingResources.h: Internal D3D12 RHI RayTracing definitions.
 =============================================================================*/
 
 #pragma once
@@ -9,11 +9,7 @@
 #include "HAL/Platform.h"
 #include "RayTracingBuiltInResources.h"
 
-#include "D3D12RHI.h"
-
-#if !defined(D3D12_MAJOR_VERSION)
-	#include "D3D12ThirdParty.h"
-#endif
+using FD3D12_GPU_VIRTUAL_ADDRESS = uint64;
 
 // Built-in local root parameters that are always bound to all hit shaders
 // Contains union for bindless and non-bindless index/vertex buffer data to make the code handling the hit group parameters easier to use
@@ -31,8 +27,15 @@ struct FD3D12HitGroupSystemParameters
 		};
 		struct
 		{
-			D3D12_GPU_VIRTUAL_ADDRESS IndexBuffer;
-			D3D12_GPU_VIRTUAL_ADDRESS VertexBuffer;
+			FD3D12_GPU_VIRTUAL_ADDRESS IndexBuffer;
+			FD3D12_GPU_VIRTUAL_ADDRESS VertexBuffer;
 		};
 	};
 };
+
+struct FD3D12RayTracingOfflineBvhHeader
+{
+	uint64 ResultSize = 0;
+	uint64 BuildScratchSize = 0;
+	uint64 UpdateScratchSize = 0;
+};
diff --git a/Engine/Source/Runtime/D3D12RHI/Private/D3D12RayTracing.cpp b/Engine/Source/Runtime/D3D12RHI/Private/D3D12RayTracing.cpp
--- a/Engine/Source/Runtime/D3D12RHI/Private/D3D12RayTracing.cpp
+++ b/Engine/Source/Runtime/D3D12RHI/Private/D3D12RayTracing.cpp
@@ -2831,16 +2831,9 @@ FD3D12RayTracingGeometry::FD3D12RayTracingGeometry(FRHICommandListBase& RHICmdLi
 
 	SetDirty(FRHIGPUMask::All(), true);
 
-	struct FOfflineBVHHeader
-	{
-		uint64 ResultSize;
-		uint64 BuildScratchSize;
-		uint64 UpdateScratchSize;
-	};
-
 	if (Initializer.OfflineData)
 	{
-		FOfflineBVHHeader* DataHeader = (FOfflineBVHHeader*)Initializer.OfflineData->GetResourceData();
+		FD3D12RayTracingOfflineBvhHeader* DataHeader = (FD3D12RayTracingOfflineBvhHeader*)Initializer.OfflineData->GetResourceData();
 
 		SizeInfo.ResultSize = DataHeader->ResultSize;
 		SizeInfo.BuildScratchSize = DataHeader->BuildScratchSize;
@@ -2891,8 +2884,8 @@ FD3D12RayTracingGeometry::FD3D12RayTracingGeometry(FRHICommandListBase& RHICmdLi
 	{
 		FD3D12Device* Device = Adapter->GetDevice(0);
 
-		const uint8* Data = ((const uint8*)Initializer.OfflineData->GetResourceData()) + sizeof(FOfflineBVHHeader);
-		uint32 Size = Initializer.OfflineData->GetResourceDataSize() - sizeof(FOfflineBVHHeader);
+		const uint8* Data = ((const uint8*)Initializer.OfflineData->GetResourceData()) + sizeof(FD3D12RayTracingOfflineBvhHeader);
+		uint32 Size = Initializer.OfflineData->GetResourceDataSize() - sizeof(FD3D12RayTracingOfflineBvhHeader);
 
 		FD3D12ResourceLocation SrcResourceLoc(Device);
 		uint8* DstDataBase = (uint8*)Adapter->GetUploadHeapAllocator(0).AllocUploadResource(Size, 256, SrcResourceLoc);
diff --git a/Engine/Source/Runtime/D3D12RHI/Private/D3D12RayTracing.h b/Engine/Source/Runtime/D3D12RHI/Private/D3D12RayTracing.h
--- a/Engine/Source/Runtime/D3D12RHI/Private/D3D12RayTracing.h
+++ b/Engine/Source/Runtime/D3D12RHI/Private/D3D12RayTracing.h
@@ -8,6 +8,8 @@
 
 #include "D3D12RayTracingResources.h"
 
+static_assert(sizeof(FD3D12_GPU_VIRTUAL_ADDRESS) == sizeof(D3D12_GPU_VIRTUAL_ADDRESS), "Size of FD3D12_GPU_VIRTUAL_ADDRESS must match D3D12_GPU_VIRTUAL_ADDRESS");
+
 class FD3D12RayTracingPipelineState;
 class FD3D12RayTracingShaderBindingTable;
 
