- Make sure element count for hitrecord bindless SRV is computed correctly and we are not loosing entries due to divide by 4 for RAW buffer view
- Always align the buffer allocation size and offset to multiple of 16 so raw views with complete data can be created

diff --git a/Engine/Source/Runtime/D3D12RHI/Private/D3D12Buffer.cpp b/Engine/Source/Runtime/D3D12RHI/Private/D3D12Buffer.cpp
--- a/Engine/Source/Runtime/D3D12RHI/Private/D3D12Buffer.cpp
+++ b/Engine/Source/Runtime/D3D12RHI/Private/D3D12Buffer.cpp
@@ -148,7 +148,7 @@ FD3D12SyncPointRef FD3D12Buffer::UploadResourceDataViaCopyQueue(FResourceArrayUp
 	check(pData);
 	{
 		TRACE_CPUPROFILER_EVENT_SCOPE(CopyToUploadMemory);
-		FMemory::Memcpy(pData, InResourceArray->GetResourceData(), BufferSize);
+	FMemory::Memcpy(pData, InResourceArray->GetResourceData(), BufferSize);
 	}
 
 	// Allocate copy queue command list and perform the copy op
@@ -371,7 +371,7 @@ void FD3D12Buffer::GetResourceDescAndAlignment(uint64 InSize, uint32 InStride, E
 {
 	ResourceDesc = CD3DX12_RESOURCE_DESC::Buffer(InSize);
 
-	// Align size to 16 so RAW buffer view can be created without loosing and data at the end when dividing num elements by 4
+	// Align size to 16 so RAW buffer view can be created without loosing any data at the end when dividing num elements by 4
 	ResourceDesc.Width = Align(ResourceDesc.Width, 16);
 
 	if (EnumHasAnyFlags(InUsage, BUF_UnorderedAccess))
