- Make sure dynamically offseted index and vertex buffers are aligned to 16 bytes
- take stride into accoun when computing raytracing index buffer offset from FNiagaraRibbonIndexBuffer

diff --git a/Engine/Plugins/FX/Niagara/Source/Niagara/Private/NiagaraRendererRibbons.cpp b/Engine/Plugins/FX/Niagara/Source/Niagara/Private/NiagaraRendererRibbons.cpp
--- a/Engine/Plugins/FX/Niagara/Source/Niagara/Private/NiagaraRendererRibbons.cpp
+++ b/Engine/Plugins/FX/Niagara/Source/Niagara/Private/NiagaraRendererRibbons.cpp
@@ -1153,7 +1153,7 @@ void FNiagaraRendererRibbons::GetDynamicRayTracingInstances(FRayTracingInstanceC
 	RayTracingInstance.InstanceTransforms.Add(FMatrix::Identity);
 	
 	RayTracingGeometry.Initializer.IndexBuffer = RenderingViewResources->IndexBuffer->IndexBufferRHI;// PerViewGeneratedData.IndexAllocation.IndexBuffer->IndexBufferRHI;
-	RayTracingGeometry.Initializer.IndexBufferOffset = RenderingViewResources->IndexBuffer->FirstIndex;//PerViewGeneratedData.IndexAllocation.FirstIndex * PerViewGeneratedData.IndexAllocation.IndexStride;
+	RayTracingGeometry.Initializer.IndexBufferOffset = RenderingViewResources->IndexBuffer->FirstIndex * RenderingViewResources->IndexBuffer->IndexBufferRHI->GetStride();//PerViewGeneratedData.IndexAllocation.FirstIndex * PerViewGeneratedData.IndexAllocation.IndexStride;
 	
 	FMeshBatch MeshBatch;
 	
diff --git a/Engine/Source/Runtime/D3D12RHI/Private/D3D12Buffer.cpp b/Engine/Source/Runtime/D3D12RHI/Private/D3D12Buffer.cpp
--- a/Engine/Source/Runtime/D3D12RHI/Private/D3D12Buffer.cpp
+++ b/Engine/Source/Runtime/D3D12RHI/Private/D3D12Buffer.cpp
@@ -372,7 +372,7 @@ void FD3D12Buffer::GetResourceDescAndAlignment(uint64 InSize, uint32 InStride, E
 	ResourceDesc = CD3DX12_RESOURCE_DESC::Buffer(InSize);
 
 	// Align size to 16 so RAW buffer view can be created without loosing any data at the end when dividing num elements by 4
-	ResourceDesc.Width = Align(ResourceDesc.Width, 16);
+	ResourceDesc.Width = Align(ResourceDesc.Width, RHI_MIN_BUFFER_ALIGNMENT);
 
 	if (EnumHasAnyFlags(InUsage, BUF_UnorderedAccess))
 	{
@@ -405,7 +405,7 @@ void FD3D12Buffer::GetResourceDescAndAlignment(uint64 InSize, uint32 InStride, E
 	else
 	{
 		// Structured buffers, non-ByteAddress buffers, need to be aligned to their stride to ensure that they can be addressed correctly with element based offsets.
-		Alignment = (InStride > 0) && (EnumHasAnyFlags(InUsage, BUF_StructuredBuffer) || !EnumHasAnyFlags(InUsage, BUF_ByteAddressBuffer | BUF_DrawIndirect)) ? FMath::LeastCommonMultiplier(InStride, 16) : 16;
+		Alignment = (InStride > 0) && (EnumHasAnyFlags(InUsage, BUF_StructuredBuffer) || !EnumHasAnyFlags(InUsage, BUF_ByteAddressBuffer | BUF_DrawIndirect)) ? FMath::LeastCommonMultiplier(InStride, RHI_MIN_BUFFER_ALIGNMENT) : RHI_MIN_BUFFER_ALIGNMENT;
 	}
 }
 
diff --git a/Engine/Source/Runtime/D3D12RHI/Private/D3D12RayTracing.cpp b/Engine/Source/Runtime/D3D12RHI/Private/D3D12RayTracing.cpp
--- a/Engine/Source/Runtime/D3D12RHI/Private/D3D12RayTracing.cpp
+++ b/Engine/Source/Runtime/D3D12RHI/Private/D3D12RayTracing.cpp
@@ -3040,7 +3040,7 @@ void FD3D12RayTracingGeometry::AllocateBufferSRVs(uint32 InGPUIndex)
 	// Procedural doesn't need any SRVs for index buffer
 	if (Initializer.IndexBuffer && Initializer.GeometryType == RTGT_Triangles)
 	{
-		checkf((Initializer.IndexBufferOffset % 16) == 0, TEXT("The byte offset of raw views must be a multiple of 16 (specified offset: %d)."), Initializer.IndexBufferOffset);
+		checkf((Initializer.IndexBufferOffset % RHI_MIN_BUFFER_ALIGNMENT) == 0, TEXT("The byte offset of raw views must be a multiple of %d (specified offset: %d)."), RHI_MIN_BUFFER_ALIGNMENT, Initializer.IndexBufferOffset);
 
 		FD3D12Buffer* IndexBuffer = FD3D12DynamicRHI::ResourceCast(Initializer.IndexBuffer.GetReference());
 
@@ -3059,7 +3059,7 @@ void FD3D12RayTracingGeometry::AllocateBufferSRVs(uint32 InGPUIndex)
 
 	for (const FRayTracingGeometrySegment& Segment : Initializer.Segments)
 	{
-		checkf((Segment.VertexBufferOffset % 16) == 0, TEXT("The byte offset of raw views must be a multiple of 16 (specified offset: %d)."), Segment.VertexBufferOffset);
+		checkf((Segment.VertexBufferOffset % RHI_MIN_BUFFER_ALIGNMENT) == 0, TEXT("The byte offset of raw views must be a multiple of %d (specified offset: %d)."), RHI_MIN_BUFFER_ALIGNMENT, Segment.VertexBufferOffset);
 
 		FD3D12Buffer* VertexBuffer = FD3D12DynamicRHI::ResourceCast(Segment.VertexBuffer.GetReference());
 
diff --git a/Engine/Source/Runtime/RHI/Public/RHIDefinitions.h b/Engine/Source/Runtime/RHI/Public/RHIDefinitions.h
--- a/Engine/Source/Runtime/RHI/Public/RHIDefinitions.h
+++ b/Engine/Source/Runtime/RHI/Public/RHIDefinitions.h
@@ -24,6 +24,9 @@
 /** The alignment in bytes between elements of array shader parameters. */
 #define SHADER_PARAMETER_ARRAY_ELEMENT_ALIGNMENT 16
 
+/** The alignment in bytes for rhi buffers. */
+#define RHI_MIN_BUFFER_ALIGNMENT 16
+
 // RHICreateUniformBuffer assumes C++ constant layout matches the shader layout when extracting float constants, yet the C++ struct contains pointers.  
 // Enforce a min size of 64 bits on pointer types in uniform buffer structs to guarantee layout matching between languages.
 #define SHADER_PARAMETER_POINTER_ALIGNMENT sizeof(uint64)
diff --git a/Engine/Source/Runtime/RenderCore/Private/GlobalRenderResources.cpp b/Engine/Source/Runtime/RenderCore/Private/GlobalRenderResources.cpp
--- a/Engine/Source/Runtime/RenderCore/Private/GlobalRenderResources.cpp
+++ b/Engine/Source/Runtime/RenderCore/Private/GlobalRenderResources.cpp
@@ -930,7 +930,7 @@ FGlobalDynamicVertexBuffer::FAllocation FGlobalDynamicVertexBuffer::Allocate(uin
 	Allocation.Buffer = VertexBuffer->MappedBuffer + VertexBuffer->AllocatedByteCount;
 	Allocation.VertexBuffer = VertexBuffer;
 	Allocation.VertexOffset = VertexBuffer->AllocatedByteCount;
-	VertexBuffer->AllocatedByteCount += SizeInBytes;
+	VertexBuffer->AllocatedByteCount += Align(SizeInBytes, RHI_MIN_BUFFER_ALIGNMENT);
 	return Allocation;
 }
 
@@ -978,7 +978,7 @@ FGlobalDynamicIndexBuffer::FAllocation FGlobalDynamicIndexBuffer::Allocate(uint3
 	Allocation.Buffer = IndexBuffer->MappedBuffer + IndexBuffer->AllocatedByteCount;
 	Allocation.IndexBuffer = IndexBuffer;
 	Allocation.FirstIndex = IndexBuffer->AllocatedByteCount / IndexStride;
-	IndexBuffer->AllocatedByteCount += SizeInBytes;
+	IndexBuffer->AllocatedByteCount += Align(SizeInBytes, RHI_MIN_BUFFER_ALIGNMENT);
 	return Allocation;
 }
 
diff --git a/Engine/Source/Runtime/Renderer/Private/RayTracing/RayTracingDynamicGeometry.cpp b/Engine/Source/Runtime/Renderer/Private/RayTracing/RayTracingDynamicGeometry.cpp
--- a/Engine/Source/Runtime/Renderer/Private/RayTracing/RayTracingDynamicGeometry.cpp
+++ b/Engine/Source/Runtime/Renderer/Private/RayTracing/RayTracingDynamicGeometry.cpp
@@ -237,12 +237,12 @@ void FRayTracingDynamicGeometryCollection::AddDynamicMeshBatchForGeometryUpdate(
 		VertexPositionBuffer->UsedSize += UpdateParams.VertexBufferSize;
 
 		// Make sure vertex buffer offset is aligned to 16 (required for Raw SRV views)
-		VertexPositionBuffer->UsedSize = Align(VertexPositionBuffer->UsedSize, 16);
+		VertexPositionBuffer->UsedSize = Align(VertexPositionBuffer->UsedSize, RHI_MIN_BUFFER_ALIGNMENT);
 
 		bUseSharedVertexBuffer = true;
 		RWBuffer = &VertexPositionBuffer->RWBuffer;
 	}
-	check(IsAligned(VertexBufferOffset, 16));
+	check(IsAligned(VertexBufferOffset, RHI_MIN_BUFFER_ALIGNMENT));
 
 	FRayTracingDynamicGeometryBuildParams GeometryBuildParams;
 	GeometryBuildParams.DispatchCommands.Reserve(UpdateParams.MeshBatches.Num());
