Fix PVS 7.33 warnings

diff --git a/Engine/Source/Runtime/D3D12RHI/Private/D3D12RayTracing.cpp b/Engine/Source/Runtime/D3D12RHI/Private/D3D12RayTracing.cpp
--- a/Engine/Source/Runtime/D3D12RHI/Private/D3D12RayTracing.cpp
+++ b/Engine/Source/Runtime/D3D12RHI/Private/D3D12RayTracing.cpp
@@ -666,8 +666,7 @@ static TRefCountPtr<ID3D12StateObject> CreateRayTracingStateObject(
 	{
 		// Enable the NVAPI extension to allow shader extensions in ray tracing for shader execution reordering
 		NvAPI_Status NvapiStatus = NvAPI_D3D12_SetNvShaderExtnSlotSpaceLocalThread(RayTracingDevice, NV_SER_EXTENSION_SLOT, NV_SER_EXTENSION_SPACE);
-		checkf(NvapiStatus == NVAPI_OK, TEXT("NVAPI ERROR %d on setup extensions\n"), NvapiStatus);
-		bReorderingDisabled = (NvapiStatus == NVAPI_OK);
+		bReorderingDisabled = ensureMsgf(NvapiStatus == NVAPI_OK, TEXT("NVAPI ERROR %d on setup extensions\n"), NvapiStatus);
 	}
 #endif
 
