[RayTracing] Improvements to offline BLAS streaming and residency management.

Context:
- currently metadata regarding offline BLAS is only available at runtime after the offline BLAS is streamed in, and even then this metadata is only available on RHI thread.
- FRayTracingGeometryManager determines which FRayTracingGeometry should be streamed/evicted based on size of BLAS if they're built at runtime, which can be significanly larger than offline built BLAS.
- streaming subsections of FRayTracingGeometry::RawData is also not currently supported which means in some cases more data is streamed than necessary.

Changes:
- Store new FRayTracingGeometryOfflineDataHeader inline instead of BulkData when serializing FRayTracingGeometry built during cooking.
- RHIGetRayTracingGeometryOfflineMetadata(...) to decode platform specific data in FRayTracingGeometryOfflineDataHeader.
- Updated RHICalcRayTracingGeometrySize(...) to take offline BLAS size into account when OfflineDataHeader is valid (which is used by FRayTracingGeometryManager).
- Updated FRayTracingGeometryManager to stream section of BulkData based on output of RHIGetRayTracingGeometryOfflineMetadata(...).

diff --git a/Engine/Source/Runtime/Core/Private/UObject/FortniteMainBranchObjectVersion.cpp b/Engine/Source/Runtime/Core/Private/UObject/FortniteMainBranchObjectVersion.cpp
--- a/Engine/Source/Runtime/Core/Private/UObject/FortniteMainBranchObjectVersion.cpp
+++ b/Engine/Source/Runtime/Core/Private/UObject/FortniteMainBranchObjectVersion.cpp
@@ -13,7 +13,7 @@ TMap<FGuid, FGuid> FFortniteMainBranchObjectVersion::GetSystemGuids()
 	SystemGuids.Add(DevGuids.NIAGARASHADERMAP_DERIVEDDATA_VER, FGuid("C417A50271B3427D98785388266E3B66"));
 	SystemGuids.Add(DevGuids.Niagara_LatestScriptCompileVersion, FGuid("F8D995224324470E97ED718D33939A0E"));
 	SystemGuids.Add(DevGuids.SkeletalMeshDerivedDataVersion, FGuid("B9AEAA3EE7AC4FCFB3860E8F3E06DD75"));
-	SystemGuids.Add(DevGuids.STATICMESH_DERIVEDDATA_VER, FGuid("3231A08A87A248A09D959335B0F4060A"));
+	SystemGuids.Add(DevGuids.STATICMESH_DERIVEDDATA_VER, FGuid("CDFDB82FE0A83DB602AE2EFFD8BB7F42"));
 	SystemGuids.Add(DevGuids.MaterialTranslationDDCVersion, FGuid("38F4505C9CB84A8B9C9648D2E4CF0ED6"));
 	return SystemGuids;
 }
diff --git a/Engine/Source/Runtime/D3D12RHI/Internal/D3D12RayTracingResources.h b/Engine/Source/Runtime/D3D12RHI/Internal/D3D12RayTracingResources.h
--- a/Engine/Source/Runtime/D3D12RHI/Internal/D3D12RayTracingResources.h
+++ b/Engine/Source/Runtime/D3D12RHI/Internal/D3D12RayTracingResources.h
@@ -35,7 +35,6 @@ struct FD3D12HitGroupSystemParameters
 
 struct FD3D12RayTracingOfflineBvhHeader
 {
-	uint64 ResultSize = 0;
-	uint64 BuildScratchSize = 0;
-	uint64 UpdateScratchSize = 0;
+	uint32 Size = 0;
+	uint32 SerializedSize = 0;
 };
diff --git a/Engine/Source/Runtime/D3D12RHI/Private/D3D12RHIPrivate.h b/Engine/Source/Runtime/D3D12RHI/Private/D3D12RHIPrivate.h
--- a/Engine/Source/Runtime/D3D12RHI/Private/D3D12RHIPrivate.h
+++ b/Engine/Source/Runtime/D3D12RHI/Private/D3D12RHIPrivate.h
@@ -469,6 +469,8 @@ public:
 	virtual FRayTracingAccelerationStructureSize RHICalcRayTracingSceneSize(const FRayTracingSceneInitializer& Initializer) final override;
 	virtual FRayTracingAccelerationStructureSize RHICalcRayTracingGeometrySize(const FRayTracingGeometryInitializer& Initializer) final override;
 
+	virtual FRayTracingAccelerationStructureOfflineMetadata RHIGetRayTracingGeometryOfflineMetadata(const FRayTracingGeometryOfflineDataHeader& OfflineDataHeader) final override;
+
 	virtual FRayTracingGeometryRHIRef RHICreateRayTracingGeometry(FRHICommandListBase& RHICmdList, const FRayTracingGeometryInitializer& Initializer) final override;
 	virtual FRayTracingSceneRHIRef RHICreateRayTracingScene(FRayTracingSceneInitializer Initializer) final override;
 	virtual FRayTracingShaderRHIRef RHICreateRayTracingShader(TArrayView<const uint8> Code, const FSHAHash& Hash, EShaderFrequency ShaderFrequency) final override;
diff --git a/Engine/Source/Runtime/D3D12RHI/Private/D3D12RayTracing.cpp b/Engine/Source/Runtime/D3D12RHI/Private/D3D12RayTracing.cpp
--- a/Engine/Source/Runtime/D3D12RHI/Private/D3D12RayTracing.cpp
+++ b/Engine/Source/Runtime/D3D12RHI/Private/D3D12RayTracing.cpp
@@ -2628,6 +2628,20 @@ void TranslateRayTracingGeometryDescs(const FRayTracingGeometryInitializer& Init
 
 FRayTracingAccelerationStructureSize FD3D12DynamicRHI::RHICalcRayTracingGeometrySize(const FRayTracingGeometryInitializer& Initializer)
 {
+	{
+		const bool bHasOfflineMetadata = Initializer.OfflineDataHeader.IsValid();
+
+		if (bHasOfflineMetadata)
+		{
+			FRayTracingAccelerationStructureOfflineMetadata OfflineBvhMetadata = RHIGetRayTracingGeometryOfflineMetadata(Initializer.OfflineDataHeader);
+
+			FRayTracingAccelerationStructureSize SizeInfo = {};
+			SizeInfo.ResultSize = Align(OfflineBvhMetadata.Size, GRHIRayTracingAccelerationStructureAlignment);
+
+			return SizeInfo;
+		}
+	}
+
 	FRayTracingAccelerationStructureSize SizeInfo = {};
 	
 	ERayTracingAccelerationStructureFlags BuildFlags = GetRayTracingAccelerationStructureBuildFlags(Initializer);
@@ -2686,6 +2700,21 @@ FRayTracingAccelerationStructureSize FD3D12DynamicRHI::RHICalcRayTracingGeometry
 	return SizeInfo;
 }
 
+FRayTracingAccelerationStructureOfflineMetadata FD3D12DynamicRHI::RHIGetRayTracingGeometryOfflineMetadata(const FRayTracingGeometryOfflineDataHeader& OfflineDataHeader)
+{
+	static_assert(sizeof(FD3D12RayTracingOfflineBvhHeader) <= sizeof(FRayTracingGeometryOfflineDataHeader), "FRayTracingGeometryOfflineDataHeader must be large enough to fit FOfflineBVHHeader");
+
+	FD3D12RayTracingOfflineBvhHeader BvhHeader;
+	FMemory::Memcpy(&BvhHeader, &OfflineDataHeader, sizeof(BvhHeader));
+
+	FRayTracingAccelerationStructureOfflineMetadata Metadata;
+	Metadata.Size = BvhHeader.Size;
+	Metadata.SerializedSize = BvhHeader.SerializedSize;
+	Metadata.SerializedOffset = 0;
+
+	return Metadata;
+}
+
 FRayTracingPipelineStateRHIRef FD3D12DynamicRHI::RHICreateRayTracingPipelineState(const FRayTracingPipelineStateInitializer& Initializer)
 {
 	FD3D12Device* Device = GetAdapter().GetDevice(0); // All pipelines are created on the first node, as they may be used on any other linked GPU.
@@ -2830,15 +2859,39 @@ FD3D12RayTracingGeometry::FD3D12RayTracingGeometry(FRHICommandListBase& RHICmdLi
 
 	SetDirty(FRHIGPUMask::All(), true);
 
-	if (Initializer.OfflineData)
+	const bool bHasOfflineMetadata = Initializer.OfflineDataHeader.IsValid();
+	FRayTracingAccelerationStructureOfflineMetadata OfflineBvhMetadata;
+
+	if (bHasOfflineMetadata)
+	{
+		OfflineBvhMetadata = RHIGetRayTracingGeometryOfflineMetadata(InInitializer.OfflineDataHeader);
+	}
+
+	const void* SourceData = nullptr;
+
+	if (Initializer.OfflineData != nullptr)
+	{
+		checkf(bHasOfflineMetadata, TEXT("OfflineData provided in Initializer has data but OfflineDataHeader is not valid."));
+
+		SourceData = Initializer.OfflineData->GetResourceData();
+
+		const uint32 SourceDataSize = Initializer.OfflineData->GetResourceDataSize();
+		checkf(SourceDataSize >= OfflineBvhMetadata.SerializedSize, TEXT("OfflineData provided in Initializer has %u bytes but FD3D12RayTracingGeometry expected %u bytes."), SourceDataSize, OfflineBvhMetadata.SerializedSize);
+	}
+
+	if (SourceData != nullptr)
 	{
-		FD3D12RayTracingOfflineBvhHeader* DataHeader = (FD3D12RayTracingOfflineBvhHeader*)Initializer.OfflineData->GetResourceData();
+		checkf(!InInitializer.bAllowUpdate, TEXT("FD3D12RayTracingGeometry doesn't support updating BVH created using offline data."));
+	}
 
-		SizeInfo.ResultSize = DataHeader->ResultSize;
-		SizeInfo.BuildScratchSize = DataHeader->BuildScratchSize;
-		SizeInfo.UpdateScratchSize = DataHeader->UpdateScratchSize;
+	// Compute the required size of the in-memory BVH buffer
+	if (SourceData != nullptr)
+	{
+		SizeInfo.ResultSize = OfflineBvhMetadata.Size;
+		SizeInfo.BuildScratchSize = 0;
+		SizeInfo.UpdateScratchSize = 0;
 
-		AccelerationStructureCompactedSize = DataHeader->ResultSize;
+		AccelerationStructureCompactedSize = OfflineBvhMetadata.Size;
 	}
 	else
 	{
@@ -2879,16 +2932,13 @@ FD3D12RayTracingGeometry::FD3D12RayTracingGeometry(FRHICommandListBase& RHICmdLi
 	INC_DWORD_STAT_BY(STAT_D3D12RayTracingTrianglesBLAS, Initializer.TotalPrimitiveCount);
 
 	const bool bForRendering = Initializer.Type == ERayTracingGeometryInitializerType::Rendering;
-	if (Initializer.OfflineData)
+	if (SourceData != nullptr)
 	{
 		FD3D12Device* Device = Adapter->GetDevice(0);
 
-		const uint8* Data = ((const uint8*)Initializer.OfflineData->GetResourceData()) + sizeof(FD3D12RayTracingOfflineBvhHeader);
-		uint32 Size = Initializer.OfflineData->GetResourceDataSize() - sizeof(FD3D12RayTracingOfflineBvhHeader);
-
 		FD3D12ResourceLocation SrcResourceLoc(Device);
-		uint8* DstDataBase = (uint8*)Adapter->GetUploadHeapAllocator(0).AllocUploadResource(Size, 256, SrcResourceLoc);
-		FMemory::Memcpy(DstDataBase, Data, Size);
+		uint8* DstDataBase = (uint8*)Adapter->GetUploadHeapAllocator(0).AllocUploadResource(OfflineBvhMetadata.SerializedSize, 256, SrcResourceLoc);
+		FMemory::Memcpy(DstDataBase, SourceData, OfflineBvhMetadata.SerializedSize);
 
 		RHICmdList.EnqueueLambda([this, SrcResourceLoc = MoveTemp(SrcResourceLoc), bForRendering](FRHICommandListBase& ExecutingCmdList)
 		{
diff --git a/Engine/Source/Runtime/Engine/Private/Rendering/RayTracingGeometryManager.cpp b/Engine/Source/Runtime/Engine/Private/Rendering/RayTracingGeometryManager.cpp
--- a/Engine/Source/Runtime/Engine/Private/Rendering/RayTracingGeometryManager.cpp
+++ b/Engine/Source/Runtime/Engine/Private/Rendering/RayTracingGeometryManager.cpp
@@ -357,10 +357,13 @@ void FRayTracingGeometryManager::SetRayTracingGeometryStreamingData(const FRayTr
 
 	checkf(RegisteredGeometries.IsValidIndex(Geometry->RayTracingGeometryHandle), TEXT("SetRayTracingGeometryStreamingData(...) can only be used with FRayTracingGeometry that has been registered with FRayTracingGeometryManager."));
 
+	FRayTracingAccelerationStructureOfflineMetadata OfflineMetadata = RHIGetRayTracingGeometryOfflineMetadata(Geometry->Initializer.OfflineDataHeader);
+	check(OfflineMetadata.SerializedSize <= Size);
+
 	FRegisteredGeometry& RegisteredGeometry = RegisteredGeometries[Geometry->RayTracingGeometryHandle];
 	RegisteredGeometry.StreamableData = &BulkData;
-	RegisteredGeometry.StreamableDataOffset = Offset;
-	RegisteredGeometry.StreamableDataSize = Size;
+	RegisteredGeometry.StreamableDataOffset = Offset + OfflineMetadata.SerializedOffset;
+	RegisteredGeometry.StreamableDataSize = OfflineMetadata.SerializedSize;
 }
 
 
@@ -1018,8 +1021,11 @@ void FRayTracingGeometryManager::ProcessCompletedStreamingRequests(FRHICommandLi
 			else
 			{
 				{
+					RegisteredGeometry.Geometry->RawData.SetNum(RegisteredGeometry.StreamableDataSize);
+
 					FMemoryReaderView Ar(PendingRequest.RequestBuffer.GetView(), /*bIsPersistent=*/ true);
-					RegisteredGeometry.Geometry->RawData.BulkSerialize(Ar);
+					Ar.Serialize(RegisteredGeometry.Geometry->RawData.GetData(), RegisteredGeometry.StreamableDataSize);
+					check(!Ar.IsError());
 				}
 
 				{
diff --git a/Engine/Source/Runtime/Engine/Private/StaticMesh.cpp b/Engine/Source/Runtime/Engine/Private/StaticMesh.cpp
--- a/Engine/Source/Runtime/Engine/Private/StaticMesh.cpp
+++ b/Engine/Source/Runtime/Engine/Private/StaticMesh.cpp
@@ -582,11 +582,12 @@ void FStaticMeshLODResources::SerializeBuffers(FArchive& Ar, UStaticMesh* OwnerS
 		AdjacencyIndexBuffer.Serialize(Ar, bNeedsCPUAccess);
 	}
 
-	FRayTracingGeometry DummyRayTracingGeometry;
-	FRayTracingGeometry* SerializedRayTracingGeometry = RayTracingGeometry != nullptr ? RayTracingGeometry : &DummyRayTracingGeometry;
-
 	if (bSerializeRayTracingGeometry)
 	{
+		FRayTracingGeometry DummyRayTracingGeometry;
+		FRayTracingGeometry* SerializedRayTracingGeometry = RayTracingGeometry != nullptr ? RayTracingGeometry : &DummyRayTracingGeometry;
+
+		Ar << SerializedRayTracingGeometry->RawDataHeader;
 		SerializedRayTracingGeometry->RawData.BulkSerialize(Ar);
 		AccumRayTracingGeometrySize(*SerializedRayTracingGeometry, OutBuffersSize.SerializedBuffersSize);
 	}
@@ -689,6 +690,12 @@ void FStaticMeshLODResources::SerializeAvailabilityInfo(FArchive& Ar)
 		FRawStaticIndexBuffer AdjacencyIndexBuffer;
 		AdjacencyIndexBuffer.SerializeMetaData(Ar);
 	}
+
+	{
+		FRayTracingGeometry DummyRayTracingGeometry;
+		FRayTracingGeometry* SerializedRayTracingGeometry = RayTracingGeometry != nullptr ? RayTracingGeometry : &DummyRayTracingGeometry;
+		Ar << SerializedRayTracingGeometry->RawDataHeader;
+	}
 }
 
 void FStaticMeshLODResources::ClearAvailabilityInfo()
@@ -1717,10 +1724,13 @@ void FStaticMeshRayTracingProxyLOD::Serialize(FArchive& Ar, UObject* Owner, int3
 
 	if (bBuffersInlined)
 	{
+		SerializeBuffersMetaData(Ar);
 		SerializeBuffers(Ar, OwnerStaticMesh, 0);
 	}
 	else
 	{
+		SerializeBuffersMetaData(Ar);
+
 #if WITH_EDITOR
 		// Serialize buffers to FByteBulkData when saving to DDC
 		// (don't need to do it again when cooking)
@@ -1775,14 +1785,34 @@ uint32 FStaticMeshRayTracingProxyLOD::SerializeBuffers(FArchive& Ar, UStaticMesh
 
 	if (bOwnsRayTracingGeometry)
 	{
-		RayTracingGeometry->RawData.BulkSerialize(Ar);
+		const uint32 RawDataSize = RayTracingGeometry->RawData.Num() * RayTracingGeometry->RawData.GetTypeSize();
+
+		// serialize RawData as a memory blob to be able to stream-in specific chunks based on FRayTracingAccelerationStructureOfflineMetadata
+		Ar.Serialize(RayTracingGeometry->RawData.GetData(), RawDataSize);
 
-		OutSize += RayTracingGeometry->RawData.Num();
+		OutSize += RawDataSize;
 	}
 
 	return OutSize;
 }
 
+void FStaticMeshRayTracingProxyLOD::SerializeBuffersMetaData(FArchive& Ar)
+{
+	if (bOwnsBuffers)
+	{
+		VertexBuffers->PositionVertexBuffer.SerializeMetaData(Ar);
+		VertexBuffers->StaticMeshVertexBuffer.SerializeMetaData(Ar);
+		VertexBuffers->ColorVertexBuffer.SerializeMetaData(Ar);
+
+		IndexBuffer->SerializeMetaData(Ar);
+	}
+
+	if (bOwnsRayTracingGeometry)
+	{
+		Ar << RayTracingGeometry->RawDataHeader;
+	}
+}
+
 /*------------------------------------------------------------------------------
 	FStaticMeshRayTracingProxy
 ------------------------------------------------------------------------------*/
diff --git a/Engine/Source/Runtime/Engine/Public/StaticMeshResources.h b/Engine/Source/Runtime/Engine/Public/StaticMeshResources.h
--- a/Engine/Source/Runtime/Engine/Public/StaticMeshResources.h
+++ b/Engine/Source/Runtime/Engine/Public/StaticMeshResources.h
@@ -710,6 +710,7 @@ struct FStaticMeshRayTracingProxyLOD
 	void Serialize(FArchive& Ar, UObject* Owner, int32 Index);
 
 	uint32 SerializeBuffers(FArchive& Ar, UStaticMesh* OwnerStaticMesh, uint8 InStripFlags);
+	void SerializeBuffersMetaData(FArchive& Ar);
 
 	/** Discard loaded data. Used when a streaming request is cancelled */
 	void DiscardCPUData();
diff --git a/Engine/Source/Runtime/RHI/Private/RHICommandList.cpp b/Engine/Source/Runtime/RHI/Private/RHICommandList.cpp
--- a/Engine/Source/Runtime/RHI/Private/RHICommandList.cpp
+++ b/Engine/Source/Runtime/RHI/Private/RHICommandList.cpp
@@ -1747,7 +1747,7 @@ void FRHIComputeCommandList::BuildAccelerationStructure(FRHIRayTracingGeometry*
 
 void FRHIComputeCommandList::BuildAccelerationStructures(const TArrayView<const FRayTracingGeometryBuildParams> Params)
 {
-	// Buffer size is limited to 2Gb, thus split acceleration structure building into pieces to accomodate this limitation
+	// Buffer size is limited to 2Gb, thus split acceleration structure building into pieces to accommodate this limitation
 	const uint64 MaxScratchMemorySize = 2147483647u;
 	const uint32 ParamTotalCount = Params.Num();
 
diff --git a/Engine/Source/Runtime/RHI/Public/DynamicRHI.h b/Engine/Source/Runtime/RHI/Public/DynamicRHI.h
--- a/Engine/Source/Runtime/RHI/Public/DynamicRHI.h
+++ b/Engine/Source/Runtime/RHI/Public/DynamicRHI.h
@@ -958,6 +958,12 @@ public:
 		return {};
 	}
 
+	virtual FRayTracingAccelerationStructureOfflineMetadata RHIGetRayTracingGeometryOfflineMetadata(const FRayTracingGeometryOfflineDataHeader& OfflineDataHeader)
+	{
+		checkNoEntry();
+		return {};
+	}
+
 	virtual FRayTracingGeometryRHIRef RHICreateRayTracingGeometry(FRHICommandListBase& RHICmdList, const FRayTracingGeometryInitializer& Initializer)
 	{
 		checkNoEntry();
diff --git a/Engine/Source/Runtime/RHI/Public/RHICommandList.h b/Engine/Source/Runtime/RHI/Public/RHICommandList.h
--- a/Engine/Source/Runtime/RHI/Public/RHICommandList.h
+++ b/Engine/Source/Runtime/RHI/Public/RHICommandList.h
@@ -5876,6 +5876,11 @@ FORCEINLINE FRayTracingAccelerationStructureSize RHICalcRayTracingGeometrySize(c
 	return GDynamicRHI->RHICalcRayTracingGeometrySize(Initializer);
 }
 
+FORCEINLINE FRayTracingAccelerationStructureOfflineMetadata RHIGetRayTracingGeometryOfflineMetadata(const FRayTracingGeometryOfflineDataHeader& OfflineDataHeader)
+{
+	return GDynamicRHI->RHIGetRayTracingGeometryOfflineMetadata(OfflineDataHeader);
+}
+
 FORCEINLINE void RHIBindDebugLabelName(FRHITexture* Texture, const TCHAR* Name)
 {
 	FRHICommandListImmediate::Get().BindDebugLabelName(Texture, Name);
diff --git a/Engine/Source/Runtime/RHI/Public/RHIResources.h b/Engine/Source/Runtime/RHI/Public/RHIResources.h
--- a/Engine/Source/Runtime/RHI/Public/RHIResources.h
+++ b/Engine/Source/Runtime/RHI/Public/RHIResources.h
@@ -3110,6 +3110,42 @@ struct FRayTracingGeometryInstance
 	ERayTracingInstanceFlags Flags = ERayTracingInstanceFlags::None;
 };
 
+struct FRayTracingGeometryOfflineDataHeader
+{
+	uint32 Reserved[6] = { 0xFFFFFFFF, 0xFFFFFFFF, 0xFFFFFFFF, 0xFFFFFFFF, 0xFFFFFFFF, 0xFFFFFFFF };
+
+	bool IsValid() const
+	{
+		return *this != FRayTracingGeometryOfflineDataHeader();
+	}
+
+	bool operator==(const FRayTracingGeometryOfflineDataHeader& Other) const
+	{
+		return Reserved[0] == Other.Reserved[0]
+			&& Reserved[1] == Other.Reserved[1]
+			&& Reserved[2] == Other.Reserved[2]
+			&& Reserved[3] == Other.Reserved[3]
+			&& Reserved[4] == Other.Reserved[4]
+			&& Reserved[5] == Other.Reserved[5];
+	}
+
+	bool operator!=(const FRayTracingGeometryOfflineDataHeader& Other) const
+	{
+		return !(*this == Other);
+	}
+
+	friend FArchive& operator<<(FArchive& Ar, FRayTracingGeometryOfflineDataHeader& Header)
+	{
+		Ar << Header.Reserved[0];
+		Ar << Header.Reserved[1];
+		Ar << Header.Reserved[2];
+		Ar << Header.Reserved[3];
+		Ar << Header.Reserved[4];
+		Ar << Header.Reserved[5];
+		return Ar;
+	}
+};
+
 enum ERayTracingGeometryType
 {
 	// Indexed or non-indexed triangle list with fixed function ray intersection.
@@ -3196,6 +3232,7 @@ public:
 
 	// Offline built geometry data. If null, the geometry will be built by the RHI at runtime.
 	FResourceArrayUploadInterface* OfflineData = nullptr;
+	FRayTracingGeometryOfflineDataHeader OfflineDataHeader;
 
 	// Pointer to an existing ray tracing geometry which the new geometry is built from.
 	FRHIRayTracingGeometry* SourceGeometry = nullptr;
@@ -3393,6 +3430,14 @@ struct FRayTracingAccelerationStructureSize
 	uint64 UpdateScratchSize = 0;
 };
 
+struct FRayTracingAccelerationStructureOfflineMetadata
+{
+	uint32 Size = 0;
+	uint32 ScratchSize = 0;
+	uint32 SerializedSize = 0;
+	uint32 SerializedOffset = 0;
+};
+
 class FRHIRayTracingAccelerationStructure
 	: public FRHIResource
 #if ENABLE_RHI_VALIDATION
diff --git a/Engine/Source/Runtime/RHI/Public/RHIValidation.h b/Engine/Source/Runtime/RHI/Public/RHIValidation.h
--- a/Engine/Source/Runtime/RHI/Public/RHIValidation.h
+++ b/Engine/Source/Runtime/RHI/Public/RHIValidation.h
@@ -1240,6 +1240,11 @@ public:
 		return RHI->RHICalcRayTracingGeometrySize(Initializer);
 	}
 
+	virtual FRayTracingAccelerationStructureOfflineMetadata RHIGetRayTracingGeometryOfflineMetadata(const FRayTracingGeometryOfflineDataHeader& OfflineDataHeader)
+	{
+		return RHI->RHIGetRayTracingGeometryOfflineMetadata(OfflineDataHeader);
+	}
+
 	virtual FShaderBindingTableRHIRef RHICreateShaderBindingTable(FRHICommandListBase& RHICmdList, const FRayTracingShaderBindingTableInitializer& Initializer) override final
 	{
 		return RHI->RHICreateShaderBindingTable(RHICmdList, Initializer);
diff --git a/Engine/Source/Runtime/RenderCore/Private/RayTracingGeometry.cpp b/Engine/Source/Runtime/RenderCore/Private/RayTracingGeometry.cpp
--- a/Engine/Source/Runtime/RenderCore/Private/RayTracingGeometry.cpp
+++ b/Engine/Source/Runtime/RenderCore/Private/RayTracingGeometry.cpp
@@ -24,6 +24,16 @@ PRAGMA_DISABLE_DEPRECATION_WARNINGS // Remove this in UE 5.6 once RayTracingGeom
 
 #if RHI_RAYTRACING
 
+void FRayTracingGeometry::SetInitializer(FRayTracingGeometryInitializer InInitializer)
+{
+	Initializer = MoveTemp(InInitializer);
+
+	if (!CVarDebugForceRuntimeBLAS.GetValueOnAnyThread())
+	{
+		Initializer.OfflineDataHeader = RawDataHeader;
+	}
+}
+
 void FRayTracingGeometry::InitRHIForStreaming(FRHIRayTracingGeometry* IntermediateGeometry, FRHIResourceReplaceBatcher& Batcher)
 {
 	checkf(IntermediateGeometry,
diff --git a/Engine/Source/Runtime/RenderCore/Public/RayTracingGeometry.h b/Engine/Source/Runtime/RenderCore/Public/RayTracingGeometry.h
--- a/Engine/Source/Runtime/RenderCore/Public/RayTracingGeometry.h
+++ b/Engine/Source/Runtime/RenderCore/Public/RayTracingGeometry.h
@@ -5,9 +5,7 @@
 #include "CoreMinimal.h"
 #include "Containers/DynamicRHIResourceArray.h"
 #include "RenderResource.h"
-#if RHI_RAYTRACING
 #include "RHI.h"
-#endif
 
 class FRHICommandListBase;
 
@@ -29,6 +27,7 @@ enum class ERTAccelerationStructureBuildPriority : uint8
 class FRayTracingGeometry : public FRenderResource
 {
 public:
+	FRayTracingGeometryOfflineDataHeader RawDataHeader;
 	TResourceArray<uint8> RawData;
 
 	//disable deprecation warnings for default constructors
@@ -95,10 +94,7 @@ public:
 	};
 	FRIEND_ENUM_CLASS_FLAGS(EGeometryStateFlags);
 
-	void SetInitializer(FRayTracingGeometryInitializer InInitializer)
-	{
-		Initializer = MoveTemp(InInitializer);
-	}
+	RENDERCORE_API void SetInitializer(FRayTracingGeometryInitializer InInitializer);
 
 	RENDERCORE_API bool IsValid() const;
 	RENDERCORE_API bool IsEvicted() const;
