New defaults for GPU crash debug helpers:
* engine breadcrumbs are on in all configurations
* DRED is off
* Aftermath is on where supported, but without markers and callstacks

Also, the GPU crash message box always suggests looking at the logs, since there will be state information in there even without -gpucrashdebugging.

#rnx

diff --git a/Engine/Source/Runtime/D3D12RHI/Private/D3D12Adapter.cpp b/Engine/Source/Runtime/D3D12RHI/Private/D3D12Adapter.cpp
--- a/Engine/Source/Runtime/D3D12RHI/Private/D3D12Adapter.cpp
+++ b/Engine/Source/Runtime/D3D12RHI/Private/D3D12Adapter.cpp
@@ -89,14 +89,14 @@ enum class ED3D12DredMode
 
 static TAutoConsoleVariable<int32> CVarD3D12EnableDRED(
 	TEXT("r.D3D12.DRED"),
-	(UE_BUILD_SHIPPING || UE_BUILD_TEST) ? 0 : 1,
+	0,
 	TEXT("Enable DRED GPU Crash debugging mode to track the current GPU state and logs information what operations the GPU executed last.")
 	TEXT("Has GPU overhead but gives the most information on the current GPU state when it crashes or hangs.\n"),
 	ECVF_RenderThreadSafe | ECVF_ReadOnly);
 
 static TAutoConsoleVariable<int32> CVarD3D12EnableLightweightDRED(
 	TEXT("r.D3D12.LightweightDRED"),
-	1,
+	0,
 	TEXT("Enable Lightweight DRED GPU Crash debugging mode to track the current GPU state and logs information what operations the GPU executed last.")
 	TEXT("Gives the basic information on the current GPU state when it crashes or hangs on all PC hardware.\n"),
 	ECVF_RenderThreadSafe | ECVF_ReadOnly);
diff --git a/Engine/Source/Runtime/D3D12RHI/Private/D3D12Util.cpp b/Engine/Source/Runtime/D3D12RHI/Private/D3D12Util.cpp
--- a/Engine/Source/Runtime/D3D12RHI/Private/D3D12Util.cpp
+++ b/Engine/Source/Runtime/D3D12RHI/Private/D3D12Util.cpp
@@ -919,14 +919,7 @@ void FD3D12DynamicRHI::OutputGPUCrashReport(FTextBuilder& ErrorMessage)
 
 	// Build the error message
 	ErrorMessage.AppendLine(LOCTEXT("GPU Crashed", "GPU Crashed or D3D Device Removed.\n"));
-	if (UE::RHI::UseGPUCrashDebugging())
-	{
-		ErrorMessage.AppendLine(LOCTEXT("GPU Crash Debugging enabled", "Check log for GPU state information."));
-	}
-	else
-	{
-		ErrorMessage.AppendLine(LOCTEXT("GPU Crash Debugging disabled", "Use -gpucrashdebugging to enable GPU state tracking."));
-	}
+	ErrorMessage.AppendLine(LOCTEXT("GPU Crash Debugging enabled", "Check log for GPU state information."));
 #if NV_AFTERMATH
 	if (AftermathResult.DumpPath.IsSet())
 	{
diff --git a/Engine/Source/Runtime/RHI/Private/RHI.cpp b/Engine/Source/Runtime/RHI/Private/RHI.cpp
--- a/Engine/Source/Runtime/RHI/Private/RHI.cpp
+++ b/Engine/Source/Runtime/RHI/Private/RHI.cpp
@@ -2253,7 +2253,7 @@ namespace UE::RHI
 
 	static TAutoConsoleVariable<int32> CVarGPUCrashDebuggingBreadcrumbs(
 		TEXT("r.GPUCrashDebugging.Breadcrumbs"),
-		0,
+		1,
 		TEXT("Enable RHI breadcrumbs, a vendor-agnostic method for determining which passes were active when a GPU crash occurs"),
 		ECVF_ReadOnly
 	);
diff --git a/Engine/Source/Runtime/RHICore/Private/RHICoreNvidiaAftermath.cpp b/Engine/Source/Runtime/RHICore/Private/RHICoreNvidiaAftermath.cpp
--- a/Engine/Source/Runtime/RHICore/Private/RHICoreNvidiaAftermath.cpp
+++ b/Engine/Source/Runtime/RHICore/Private/RHICoreNvidiaAftermath.cpp
@@ -46,7 +46,7 @@ namespace UE::RHICore::Nvidia::Aftermath
 
 	static TAutoConsoleVariable<int32> CVarAftermath_Callstack(
 		TEXT("r.GPUCrashDebugging.Aftermath.Callstack"),
-		1,
+		0,
 		TEXT("Enable callstack capture in Aftermath dumps"),
 		ECVF_ReadOnly
 	);
