- Fix missing D3D12_RESOURCE_FLAG_ALLOW_UNORDERED_ACCESS when allocating RayTracingDispatchRaysDescBuffer
- Without this we will hit a check on boot if buffer pooling is disabled

#rnx

diff --git a/Engine/Source/Runtime/D3D12RHI/Private/D3D12Device.cpp b/Engine/Source/Runtime/D3D12RHI/Private/D3D12Device.cpp
--- a/Engine/Source/Runtime/D3D12RHI/Private/D3D12Device.cpp
+++ b/Engine/Source/Runtime/D3D12RHI/Private/D3D12Device.cpp
@@ -395,7 +395,7 @@ void FD3D12Device::SetupAfterDeviceCreation()
 	check(RayTracingCompactionRequestHandler == nullptr);
 	RayTracingCompactionRequestHandler = new FD3D12RayTracingCompactionRequestHandler(this);
 
-	D3D12_RESOURCE_DESC DispatchRaysDescBufferDesc = CD3DX12_RESOURCE_DESC::Buffer(sizeof(D3D12_DISPATCH_RAYS_DESC), D3D12RHI_RESOURCE_FLAG_ALLOW_INDIRECT_BUFFER);
+	D3D12_RESOURCE_DESC DispatchRaysDescBufferDesc = CD3DX12_RESOURCE_DESC::Buffer(sizeof(D3D12_DISPATCH_RAYS_DESC), D3D12RHI_RESOURCE_FLAG_ALLOW_INDIRECT_BUFFER | D3D12_RESOURCE_FLAG_ALLOW_UNORDERED_ACCESS);
 	RayTracingDispatchRaysDescBuffer = GetParentAdapter()->CreateRHIBuffer(
 		DispatchRaysDescBufferDesc, 256,
 		FRHIBufferDesc(DispatchRaysDescBufferDesc.Width, 0, BUF_DrawIndirect),
