Fix hangs on D3D12 GPU crashes
 - A deadlock can occur if an engine thread takes the GD3DCallFailedCS lock inside VerifyD3D12Result and then signals the interrupt thread. If the interrupt thread also sees the failed HRESULT, it'll deadlock on the same GD3DCallFailedCS when it also enters VerifyD3D12Result. This results in a hang rather than a proper GPU crash report.
 - Refactored the crash handling functions to cope with being called from the interrupt thread. This also means we can use VERIFYD3DRESULT directly on GetDeviceRemovedReason().

Removed unused checkComRefCount and VerifyComRefCount functions
 - Leftovers that were copied from D3D11RHI

diff --git a/Engine/Source/Runtime/D3D12RHI/Private/D3D12RHIPrivate.h b/Engine/Source/Runtime/D3D12RHI/Private/D3D12RHIPrivate.h
--- a/Engine/Source/Runtime/D3D12RHI/Private/D3D12RHIPrivate.h
+++ b/Engine/Source/Runtime/D3D12RHI/Private/D3D12RHIPrivate.h
@@ -170,6 +170,13 @@ public:
 	// Dump & Log all the information we have on a GPU crash (NvAfterMath, DRED, Breadcrumbs, ...)
 	void OutputGPUCrashReport(FTextBuilder& ErrorMessage);
 
+	// Returns true if the calling thread is the interrupt thread
+	// (or acting as the interrupt thread when D3D12_USE_INTERRUPT_THREAD is zero).
+	bool IsInInterruptThread() const;
+
+	// Main crash function for handling D3D API errors
+	void HandleFailedD3D12Result(HRESULT D3DResult, ID3D12Device* Device, bool bCreatingTextures, const TCHAR* Message);
+
 private:
 	// Calls OutputGpuCrashReport, and also forces a quit.
 	void TerminateOnGPUCrash();
@@ -181,6 +188,10 @@ private:
 	class FD3D12Thread* SubmissionThread = nullptr;
 	class FD3D12Thread* InterruptThread = nullptr;
 
+	// Used when a dedicated interrupt thread is disabled (D3D12_USE_INTERRUPT_THREAD)
+	// Tracks the thread ID of the thread currently acting as the interrupt thread.
+	uint32 InterruptThreadID = 0;
+
 	enum class EQueueStatus
 	{
 		None      = 0,
diff --git a/Engine/Source/Runtime/D3D12RHI/Private/D3D12Submission.cpp b/Engine/Source/Runtime/D3D12RHI/Private/D3D12Submission.cpp
--- a/Engine/Source/Runtime/D3D12RHI/Private/D3D12Submission.cpp
+++ b/Engine/Source/Runtime/D3D12RHI/Private/D3D12Submission.cpp
@@ -87,6 +87,11 @@ public:
 		Thread->WaitForCompletion();
 	}
 
+	uint32 GetThreadID() const
+	{
+		return Thread->GetThreadID();
+	}
+
 private:
 	virtual uint32 Run() override
 	{
@@ -1145,6 +1150,8 @@ void FD3D12DynamicRHI::ProcessInterruptQueueUntil(FGraphEvent* GraphEvent)
 		Retry:
 			if (InterruptCS.TryLock())
 			{
+				TGuardValue<uint32> Guard(InterruptThreadID, FPlatformTLS::GetCurrentThreadId());
+				
 				FProcessResult Result;
 				do { Result = ProcessInterruptQueue(); }
 				// If we have a sync point, keep processing until the sync point is signaled.
@@ -1187,6 +1194,7 @@ void FD3D12DynamicRHI::ProcessInterruptQueueOnGPUCrash()
 		// If we have no interrupt thread, assume ownership on the current thread
 		// (or block forever on the scope lock if multiple threads enter this function).
 		FScopeLock Lock(&InterruptCS);
+		TGuardValue<uint32> Guard(InterruptThreadID, FPlatformTLS::GetCurrentThreadId());
 
 		while (true)
 		{
@@ -1195,6 +1203,23 @@ void FD3D12DynamicRHI::ProcessInterruptQueueOnGPUCrash()
 	}
 }
 
+bool FD3D12DynamicRHI::IsInInterruptThread() const
+{
+	uint32 ThisThreadID = FPlatformTLS::GetCurrentThreadId();
+
+	// If we don't have a dedicated interrupt thread, the thread currently acting
+	// as the interrupt thread is tracked via the InterruptThreadID field.
+
+	if (InterruptThread)
+	{
+		return ThisThreadID == InterruptThread->GetThreadID();
+	}
+	else
+	{
+		return ThisThreadID == InterruptThreadID;
+	}
+}
+
 D3D12_QUERY_DATA_PIPELINE_STATISTICS& operator += (D3D12_QUERY_DATA_PIPELINE_STATISTICS& LHS, D3D12_QUERY_DATA_PIPELINE_STATISTICS const& RHS)
 {
 	LHS.IAVertices	  += RHS.IAVertices;
@@ -1241,13 +1266,7 @@ FD3D12DynamicRHI::FProcessResult FD3D12DynamicRHI::ProcessInterruptQueue()
 
 	auto CheckForDeviceRemoved = [this](FD3D12Queue& Queue)
 	{
-		// If we get an error code here, we can't pass it directly to VERIFYD3D12RESULT, because that expects DXGI_ERROR_DEVICE_REMOVED, DXGI_ERROR_DEVICE_RESET etc. and wants to obtain the reason code itself
-		// by calling GetDeviceRemovedReason (again).
-		HRESULT DeviceRemovedReason = Queue.Device->GetDevice()->GetDeviceRemovedReason();
-		if (DeviceRemovedReason != S_OK)
-		{
-			TerminateOnGPUCrash();
-		}
+		VERIFYD3D12RESULT(Queue.Device->GetDevice()->GetDeviceRemovedReason());
 	};
 
 	FProcessResult Result;
diff --git a/Engine/Source/Runtime/D3D12RHI/Private/D3D12Util.cpp b/Engine/Source/Runtime/D3D12RHI/Private/D3D12Util.cpp
--- a/Engine/Source/Runtime/D3D12RHI/Private/D3D12Util.cpp
+++ b/Engine/Source/Runtime/D3D12RHI/Private/D3D12Util.cpp
@@ -961,80 +961,106 @@ void FD3D12DynamicRHI::TerminateOnGPUCrash()
 	FPlatformMisc::RequestExit(true, TEXT("D3D12Util.TerminateOnGPUCrash"));
 }
 
-namespace D3D12RHI
+// It's possible for multiple threads to catch GPU crashes or other D3D errors at the same time. Make sure we only log the error once
+// by acquiring this critical section inside HandleFailedD3D12Result (and never releasing it, because those functions don't return).
+static FCriticalSection GD3DCallFailedCS;
+
+void FD3D12DynamicRHI::HandleFailedD3D12Result(HRESULT D3DResult, ID3D12Device* Device, bool bCreatingTextures, const TCHAR* Message)
 {
-	// It's possible for multiple threads to catch GPU crashes or other D3D errors at the same time. Make sure we only log the error once by acquiring
-	// this critical section inside VerifyD3D12Result and VerifyD3D12CreateTextureResult (and never releasing it, because those functions don't return).
-	static FCriticalSection GD3DCallFailedCS;
+	auto LogMessage = [Message](bool bFatal)
+	{
+		if (bFatal)
+		{
+			UE_LOG(LogD3D12RHI, Fatal, TEXT("%s"), Message);
+		}
+		else
+		{
+			UE_LOG(LogD3D12RHI, Error, TEXT("%s"), Message);
+		}
+	};
 
-	void VerifyD3D12Result(HRESULT D3DResult, const ANSICHAR* Code, const ANSICHAR* Filename, uint32 Line, ID3D12Device* Device, FString Message)
+	if (D3DResult == E_OUTOFMEMORY)
 	{
-		check(FAILED(D3DResult));
-		
+		// Which ever thread wins the race gets to log the OOM error
 		GD3DCallFailedCS.Lock();
+		LogMessage(false);
 
-		const FString& ErrorString = GetD3D12ErrorString(D3DResult, Device);
-		UE_LOG(LogD3D12RHI, Error, TEXT("%s failed \n at %s:%u \n with error %s\n%s"), ANSI_TO_TCHAR(Code), ANSI_TO_TCHAR(Filename), Line, *ErrorString, *Message);
-		
-		if (D3DResult == E_OUTOFMEMORY)
+		// This function does not return.
+		TerminateOnOutOfMemory(Device, D3DResult, bCreatingTextures);
+	}
+	else if (D3DResult == DXGI_ERROR_DEVICE_REMOVED || D3DResult == DXGI_ERROR_DEVICE_HUNG || D3DResult == DXGI_ERROR_DEVICE_RESET)
+	{
+		if (IsInInterruptThread())
 		{
-			FD3D12DynamicRHI::GetD3DRHI()->TerminateOnOutOfMemory(Device, D3DResult, false);
+			//
+			// We're already on the interrupt thread. We must not block on the lock, as the interrupt thread is responsible
+			// for reporting the GPU crash (we can only safely access the RHI breadcrumbs / active payloads from the interrupt thread).
+			//
+			// Attempt to take the crash lock to prevent other threads from log spamming, but don't give up if we fail to do so.
+			//
+			if (GD3DCallFailedCS.TryLock())
+			{
+				// We're the first thread to take the lock... log the error
+				LogMessage(false);
+			}
+
+			// This function does not return.
+			TerminateOnGPUCrash();
 		}
-		else if (D3DResult == DXGI_ERROR_DEVICE_REMOVED || D3DResult == DXGI_ERROR_DEVICE_HUNG || D3DResult == DXGI_ERROR_DEVICE_RESET)
+		else
 		{
-			FD3D12DynamicRHI::GetD3DRHI()->ProcessInterruptQueueOnGPUCrash();
-		}
+			// Take the lock to ensure we report the error only once
+			GD3DCallFailedCS.Lock();
+			LogMessage(false);
 
-		// Make sure the log is flushed!
-		GLog->Panic();
+			// The interrupt thread must be the one to handle DXGI_ERROR_DEVICE_REMOVED etc.
+			// This function does not return.
+			ProcessInterruptQueueOnGPUCrash();
+		}
+	}
+	else
+	{
+		// For all other errors, take the lock to make sure we only report once
+		GD3DCallFailedCS.Lock();
+		LogMessage(false);
+	}
 
-		UE_LOG(LogD3D12RHI, Fatal, TEXT("%s failed \n at %s:%u \n with error %s\n%s"), ANSI_TO_TCHAR(Code), ANSI_TO_TCHAR(Filename), Line, *ErrorString, *Message);
+	//
+	// We'll end up here for any D3D error not covered above, or if any of those functions happen to return (they shouldn't do).
+	//
 
-		// Force shutdown, we can't do anything useful anymore.
-		FPlatformMisc::RequestExit(true, TEXT("D3D12Util.VerifyD3D12Result"));
-	}
+	// Make sure the log is flushed!
+	GLog->Panic();
 
-	void VerifyD3D12CreateTextureResult(HRESULT D3DResult, const ANSICHAR* Code, const ANSICHAR* Filename, uint32 Line, const D3D12_RESOURCE_DESC& TextureDesc, ID3D12Device* Device)
-	{
-		check(FAILED(D3DResult));
+	// Make one final (fatal) log attempt
+	LogMessage(true);
 
-		GD3DCallFailedCS.Lock();
+	// Force shutdown, we can't do anything useful anymore.
+	FPlatformMisc::RequestExit(true, TEXT("D3D12Util.HandleFailedD3D12Result"));
+}
 
-		const FString ErrorString = GetD3D12ErrorString(D3DResult, nullptr);
-		const TCHAR* D3DFormatString = UE::DXGIUtilities::GetFormatString(TextureDesc.Format);
+namespace D3D12RHI
+{
+	void VerifyD3D12Result(HRESULT D3DResult, const ANSICHAR* Code, const ANSICHAR* Filename, uint32 Line, ID3D12Device* Device, FString Message)
+	{
+		const FString& ErrorString = GetD3D12ErrorString(D3DResult, Device);
 
-		UE_LOG(LogD3D12RHI, Error,
-			TEXT("%s failed \n at %s:%u \n with error %s, \n Size=%ix%ix%i Format=%s(0x%08X), NumMips=%i, Flags=%s"),
+		FD3D12DynamicRHI::GetD3DRHI()->HandleFailedD3D12Result(D3DResult, Device, false, *FString::Printf(
+			TEXT("%s failed \n at %s:%u \n with error %s\n%s"), 
 			ANSI_TO_TCHAR(Code),
 			ANSI_TO_TCHAR(Filename),
 			Line,
 			*ErrorString,
-			TextureDesc.Width,
-			TextureDesc.Height,
-			TextureDesc.DepthOrArraySize,
-			D3DFormatString,
-			TextureDesc.Format,
-			TextureDesc.MipLevels,
-			*GetD3D12TextureFlagString(TextureDesc.Flags));
-
-		// Terminate with device removed but we don't have any GPU crash dump information
-		if (D3DResult == DXGI_ERROR_DEVICE_REMOVED || D3DResult == DXGI_ERROR_DEVICE_HUNG || D3DResult == DXGI_ERROR_DEVICE_RESET)
-		{
-			FD3D12DynamicRHI::GetD3DRHI()->ProcessInterruptQueueOnGPUCrash();
-		}
-		else if (D3DResult == E_OUTOFMEMORY)
-		{
-			FD3D12DynamicRHI::GetD3DRHI()->TerminateOnOutOfMemory(Device, D3DResult, true);
-
-#if STATS
-			GetRendererModule().DebugLogOnCrash();
-#endif // STATS
-		}
+			*Message
+		));
+	}
 
-		// Make sure the log is flushed!
-		GLog->Panic();
+	void VerifyD3D12CreateTextureResult(HRESULT D3DResult, const ANSICHAR* Code, const ANSICHAR* Filename, uint32 Line, const D3D12_RESOURCE_DESC& TextureDesc, ID3D12Device* Device)
+	{
+		const FString ErrorString = GetD3D12ErrorString(D3DResult, Device);
+		const TCHAR* D3DFormatString = UE::DXGIUtilities::GetFormatString(TextureDesc.Format);
 
-		UE_LOG(LogD3D12RHI, Fatal,
+		FD3D12DynamicRHI::GetD3DRHI()->HandleFailedD3D12Result(D3DResult, Device, true, *FString::Printf(
 			TEXT("%s failed \n at %s:%u \n with error %s, \n Size=%ix%ix%i Format=%s(0x%08X), NumMips=%i, Flags=%s"),
 			ANSI_TO_TCHAR(Code),
 			ANSI_TO_TCHAR(Filename),
@@ -1046,37 +1072,8 @@ namespace D3D12RHI
 			D3DFormatString,
 			TextureDesc.Format,
 			TextureDesc.MipLevels,
-			*GetD3D12TextureFlagString(TextureDesc.Flags));
-
-		// Force shutdown, we can't do anything useful anymore.
-		FPlatformMisc::RequestExit(true, TEXT("D3D12Util.VerifyD3D12CreateTextureResult"));
-	}
-
-	void VerifyComRefCount(IUnknown* Object, int32 ExpectedRefs, const TCHAR* Code, const TCHAR* Filename, int32 Line)
-	{
-		int32 NumRefs;
-
-		if (Object)
-		{
-			Object->AddRef();
-			NumRefs = Object->Release();
-
-			checkSlow(NumRefs != ExpectedRefs);
-
-			if (NumRefs != ExpectedRefs)
-			{
-				UE_LOG(
-					LogD3D12RHI,
-					Error,
-					TEXT("%s:(%d): %s has %d refs, expected %d"),
-					Filename,
-					Line,
-					Code,
-					NumRefs,
-					ExpectedRefs
-					);
-			}
-		}
+			*GetD3D12TextureFlagString(TextureDesc.Flags)
+		));
 	}
 }
 
diff --git a/Engine/Source/Runtime/D3D12RHI/Private/D3D12Util.h b/Engine/Source/Runtime/D3D12RHI/Private/D3D12Util.h
--- a/Engine/Source/Runtime/D3D12RHI/Private/D3D12Util.h
+++ b/Engine/Source/Runtime/D3D12RHI/Private/D3D12Util.h
@@ -42,12 +42,6 @@ namespace D3D12RHI
 #define VERIFYD3D12RESULT(x)			{HRESULT hres = x; if (FAILED(hres)) { VerifyD3D12Result(hres, #x, __FILE__, __LINE__, nullptr); }}
 #define VERIFYD3D12CREATETEXTURERESULT(x, Desc, Device) {HRESULT hres = x; if (FAILED(hres)) { VerifyD3D12CreateTextureResult(hres, #x, __FILE__, __LINE__, Desc, Device); }}
 
-	/**
-	 * Checks that a COM object has the expected number of references.
-	 */
-	extern void VerifyComRefCount(IUnknown* Object, int32 ExpectedRefs, const TCHAR* Code, const TCHAR* Filename, int32 Line);
-#define checkComRefCount(Obj,ExpectedRefs) VerifyComRefCount(Obj,ExpectedRefs,TEXT(#Obj),TEXT(__FILE__),__LINE__)
-
 	/** Checks if given GPU virtual address corresponds to any known resource allocations and logs results */
 	void LogPageFaultData(class FD3D12Adapter* InAdapter, FD3D12Device* InDevice, D3D12_GPU_VIRTUAL_ADDRESS InPageFaultAddress);
 	
