Removing ensure since it's not correct anymore.

diff --git a/Engine/Source/Runtime/D3D12RHI/Private/D3D12DescriptorCache.cpp b/Engine/Source/Runtime/D3D12RHI/Private/D3D12DescriptorCache.cpp
--- a/Engine/Source/Runtime/D3D12RHI/Private/D3D12DescriptorCache.cpp
+++ b/Engine/Source/Runtime/D3D12RHI/Private/D3D12DescriptorCache.cpp
@@ -864,9 +864,6 @@ bool FD3D12DescriptorCache::SwitchToNewBindlessResourceHeap(FD3D12DescriptorHeap
 
 		// Switch to the new heaps
 		bSetNewHeaps = SetDescriptorHeaps();
-
-		// If we didn't change heaps, then the caller sent us the wrong heap.
-		ensure(bSetNewHeaps);
 	}
 
 	return bSetNewHeaps;
