Restore direct call to TerminateOnGPUCrash as VERIFYD3D12RESULT might skip dumping GPU crashes for HRESULTS other than DXGI_ERROR_DEVICE_REMOVED, DXGI_ERROR_DEVICE_HUNG and DXGI_ERROR_DEVICE_RESET

diff --git a/Engine/Source/Runtime/D3D12RHI/Private/D3D12Submission.cpp b/Engine/Source/Runtime/D3D12RHI/Private/D3D12Submission.cpp
--- a/Engine/Source/Runtime/D3D12RHI/Private/D3D12Submission.cpp
+++ b/Engine/Source/Runtime/D3D12RHI/Private/D3D12Submission.cpp
@@ -1266,7 +1266,13 @@ FD3D12DynamicRHI::FProcessResult FD3D12DynamicRHI::ProcessInterruptQueue()
 
 	auto CheckForDeviceRemoved = [this](FD3D12Queue& Queue)
 	{
-		VERIFYD3D12RESULT(Queue.Device->GetDevice()->GetDeviceRemovedReason());
+		// If we get an error code here, we can't pass it directly to VERIFYD3D12RESULT, because that expects DXGI_ERROR_DEVICE_REMOVED,
+		// DXGI_ERROR_DEVICE_RESET etc. and wants to obtain the reason code itself by calling GetDeviceRemovedReason (again).
+		HRESULT DeviceRemovedReason = Queue.Device->GetDevice()->GetDeviceRemovedReason();
+		if (DeviceRemovedReason != S_OK)
+		{
+			TerminateOnGPUCrash();
+		}
 	};
 
 	FProcessResult Result;
