Add explicit validation for reserved resource pointers when processing commit operations in D3D12 RHI

diff --git a/Engine/Source/Runtime/D3D12RHI/Private/D3D12CommandContext.cpp b/Engine/Source/Runtime/D3D12RHI/Private/D3D12CommandContext.cpp
--- a/Engine/Source/Runtime/D3D12RHI/Private/D3D12CommandContext.cpp
+++ b/Engine/Source/Runtime/D3D12RHI/Private/D3D12CommandContext.cpp
@@ -377,6 +377,8 @@ void FD3D12ContextCommon::SetReservedBufferCommitSize(FD3D12Buffer* Buffer, uint
 	CommitDesc.Resource = Buffer->GetResource();
 	CommitDesc.CommitSizeInBytes = CommitSizeInBytes;
 
+	checkf(CommitDesc.Resource, TEXT("FD3D12CommitReservedResourceDesc::Resource must be set"));
+
 	GetPayload(EPhase::UpdateReservedResources)->ReservedResourcesToCommit.Add(CommitDesc);
 }
 
diff --git a/Engine/Source/Runtime/D3D12RHI/Private/D3D12RHIPrivate.h b/Engine/Source/Runtime/D3D12RHI/Private/D3D12RHIPrivate.h
--- a/Engine/Source/Runtime/D3D12RHI/Private/D3D12RHIPrivate.h
+++ b/Engine/Source/Runtime/D3D12RHI/Private/D3D12RHIPrivate.h
@@ -215,6 +215,9 @@ private:
 	// Ensures all prior batched command lists have reached the driver ID3D12Queue object.
 	void FlushBatchedPayloads(FD3D12Queue::FPayloadArray& PayloadsToHandDown);
 
+	// Executes reserved resource commit operations and cross-queue sync when required
+	void UpdateReservedResources(FD3D12Payload* Payload);
+
 	FCriticalSection SubmissionCS;
 	FCriticalSection InterruptCS;
 
diff --git a/Engine/Source/Runtime/D3D12RHI/Private/D3D12Submission.cpp b/Engine/Source/Runtime/D3D12RHI/Private/D3D12Submission.cpp
--- a/Engine/Source/Runtime/D3D12RHI/Private/D3D12Submission.cpp
+++ b/Engine/Source/Runtime/D3D12RHI/Private/D3D12Submission.cpp
@@ -852,6 +852,43 @@ uint64 FD3D12Queue::FinalizePayload(bool bRequiresSignal, FPayloadArray& Payload
 	return Fence.NextCompletionValue;
 }
 
+void FD3D12DynamicRHI::UpdateReservedResources(FD3D12Payload* Payload)
+{
+	FD3D12Queue& Queue = Payload->Queue;
+
+	// On some devices, some queues cannot perform tile remapping operations.
+	// We can work around this limitation by running the remapping in lockstep on another queue:
+	// - tile mapping queue waits for commands on this queue to finish
+	// - tile mapping queue performs the commit/decommit operations
+	// - this queue waits for tile mapping queue to finish
+	// The extra sync is not required when the current queue is capable of the remapping operations.
+
+	ID3D12CommandQueue* TileMappingQueue = (Queue.bSupportsTileMapping ? Queue.D3DCommandQueue : Queue.Device->TileMappingQueue).GetReference();
+	FD3D12Fence& TileMappingFence = Queue.Device->TileMappingFence;
+
+	const bool bCrossQueueSyncRequired = TileMappingQueue != Queue.D3DCommandQueue.GetReference();
+
+	if (bCrossQueueSyncRequired)
+	{
+		// tile mapping queue waits for commands on this queue to finish
+		Queue.D3DCommandQueue->Signal(TileMappingFence.D3DFence, ++TileMappingFence.LastSignaledValue);
+		TileMappingQueue->Wait(TileMappingFence.D3DFence, TileMappingFence.LastSignaledValue);
+	}
+
+	for (const FD3D12CommitReservedResourceDesc& CommitDesc : Payload->ReservedResourcesToCommit)
+	{
+		checkf(CommitDesc.Resource, TEXT("FD3D12CommitReservedResourceDesc::Resource must be set"));
+		CommitDesc.Resource->CommitReservedResource(TileMappingQueue, CommitDesc.CommitSizeInBytes);
+	}
+
+	if (bCrossQueueSyncRequired)
+	{
+		// this queue waits for tile mapping operations to finish
+		TileMappingQueue->Signal(TileMappingFence.D3DFence, ++TileMappingFence.LastSignaledValue);
+		Queue.D3DCommandQueue->Wait(TileMappingFence.D3DFence, TileMappingFence.LastSignaledValue);
+	}
+}
+
 void FD3D12DynamicRHI::FlushBatchedPayloads(FD3D12Queue::FPayloadArray& PayloadsToSubmit)
 {
 	uint32 FirstPayload = 0, LastPayload = 0;
@@ -882,42 +919,6 @@ void FD3D12DynamicRHI::FlushBatchedPayloads(FD3D12Queue::FPayloadArray& Payloads
 		}
 	};
 
-	auto UpdateReservedResources = [this](FD3D12Payload* Payload)
-	{
-		FD3D12Queue& Queue = Payload->Queue;
-
-		// On some devices, some queues cannot perform tile remapping operations.
-		// We can work around this limitation by running the remapping in lockstep on another queue:
-		// - tile mapping queue waits for commands on this queue to finish
-		// - tile mapping queue performs the commit/decommit operations
-		// - this queue waits for tile mapping queue to finish
-		// The extra sync is not required when the current queue is capable of the remapping operations.
-
-		ID3D12CommandQueue* TileMappingQueue = (Queue.bSupportsTileMapping ? Queue.D3DCommandQueue : Queue.Device->TileMappingQueue).GetReference();
-		FD3D12Fence& TileMappingFence = Queue.Device->TileMappingFence;
-
-		const bool bCrossQueueSyncRequired = TileMappingQueue != Queue.D3DCommandQueue.GetReference();
-
-		if (bCrossQueueSyncRequired)
-		{
-			// tile mapping queue waits for commands on this queue to finish
-			Queue.D3DCommandQueue->Signal(TileMappingFence.D3DFence, ++TileMappingFence.LastSignaledValue);
-			TileMappingQueue->Wait(TileMappingFence.D3DFence, TileMappingFence.LastSignaledValue);
-		}
-
-		for (const FD3D12CommitReservedResourceDesc& CommitDesc : Payload->ReservedResourcesToCommit)
-		{
-			CommitDesc.Resource->CommitReservedResource(TileMappingQueue, CommitDesc.CommitSizeInBytes);
-		}
-
-		if (bCrossQueueSyncRequired)
-		{
-			// this queue waits for tile mapping operations to finish
-			TileMappingQueue->Signal(TileMappingFence.D3DFence, ++TileMappingFence.LastSignaledValue);
-			Queue.D3DCommandQueue->Wait(TileMappingFence.D3DFence, TileMappingFence.LastSignaledValue);
-		}
-	};
-
 	auto Flush = [&]()
 	{
 		if (FirstPayload == LastPayload)
