Improve check message with shader name when ShaderBindingLayout doesn't match

diff --git a/Engine/Source/Runtime/D3D12RHI/Private/D3D12RayTracing.cpp b/Engine/Source/Runtime/D3D12RHI/Private/D3D12RayTracing.cpp
--- a/Engine/Source/Runtime/D3D12RHI/Private/D3D12RayTracing.cpp
+++ b/Engine/Source/Runtime/D3D12RHI/Private/D3D12RayTracing.cpp
@@ -1991,7 +1991,7 @@ public:
 		{
 			// verify that that the same shader binding layout is used for all shaders in the RTPSO or not sampling any resources
 			uint32 TotalResourceCount = Shader->ResourceCounts.NumCBs + Shader->ResourceCounts.NumSRVs + Shader->ResourceCounts.NumUAVs + Shader->ResourceCounts.NumSamplers;
-			checkf(TotalResourceCount == 0 || Shader->ShaderBindingLayoutHash == ShaderBindingLayoutHash, TEXT("Raytracing shader with with entry point %s doesn't match the RTPSO ShaderBindingLayout"), *Shader->EntryPoint);
+			checkf(TotalResourceCount == 0 || Shader->ShaderBindingLayoutHash == ShaderBindingLayoutHash, TEXT("Raytracing shader with with entry point %s and name %s doesn't match the RTPSO ShaderBindingLayout"), *Shader->EntryPoint, Shader->GetShaderName());
 
 			bool bIsAlreadyInSet = false;
 			const uint64 ShaderHash = GetShaderHash64(Shader);
