Use segment first primitive and num primitives to setup correct max primitive count when creating index buffer SRVs to make sure it's not going out of range

diff --git a/Engine/Source/Runtime/D3D12RHI/Private/D3D12RayTracing.cpp b/Engine/Source/Runtime/D3D12RHI/Private/D3D12RayTracing.cpp
--- a/Engine/Source/Runtime/D3D12RHI/Private/D3D12RayTracing.cpp
+++ b/Engine/Source/Runtime/D3D12RHI/Private/D3D12RayTracing.cpp
@@ -3093,12 +3093,21 @@ void FD3D12RayTracingGeometry::AllocateBufferSRVs(uint32 InGPUIndex)
 
 		FD3D12Buffer* IndexBuffer = FD3D12DynamicRHI::ResourceCast(Initializer.IndexBuffer.GetReference());
 
+		// Initializer.TotalPrimitiveCount currently just accumulated the num primitives of the segments
+		// but can be invalid if multiple segments overlap/share the same index buffer range
+		uint32 MaxPrimitiveCount = 0;
+		for (const FRayTracingGeometrySegment& Segment : Initializer.Segments)
+		{
+			MaxPrimitiveCount = FMath::Max(MaxPrimitiveCount, Segment.FirstPrimitive + Segment.NumPrimitives);
+		}
+		ensure(MaxPrimitiveCount <= Initializer.TotalPrimitiveCount);
+
 		D3D12_SHADER_RESOURCE_VIEW_DESC SRVDesc = {};
 		SRVDesc.ViewDimension = D3D12_SRV_DIMENSION_BUFFER;
 		SRVDesc.Format = DXGI_FORMAT_R32_TYPELESS;
 		SRVDesc.Shader4ComponentMapping = D3D12_DEFAULT_SHADER_4_COMPONENT_MAPPING;
 		SRVDesc.Buffer.FirstElement = (Initializer.IndexBufferOffset + IndexBuffer->ResourceLocation.GetOffsetFromBaseOfResource()) >> 2u;
-		SRVDesc.Buffer.NumElements = FMath::Max((uint32)1, ((Initializer.TotalPrimitiveCount * 3 * IndexBuffer->GetStride()) + 3) >> 2u);
+		SRVDesc.Buffer.NumElements = FMath::Max((uint32)1, ((MaxPrimitiveCount * 3 * IndexBuffer->GetStride()) + 3) >> 2u);
 		SRVDesc.Buffer.Flags = D3D12_BUFFER_SRV_FLAG_RAW;
 		SRVDesc.Buffer.StructureByteStride = 0;
 
