Set the sampler descriptor heap first to make sure it's not null if the view heap needs to increase its size.

diff --git a/Engine/Source/Runtime/D3D12RHI/Private/D3D12DescriptorCache.cpp b/Engine/Source/Runtime/D3D12RHI/Private/D3D12DescriptorCache.cpp
--- a/Engine/Source/Runtime/D3D12RHI/Private/D3D12DescriptorCache.cpp
+++ b/Engine/Source/Runtime/D3D12RHI/Private/D3D12DescriptorCache.cpp
@@ -159,14 +159,14 @@ void FD3D12DescriptorCache::OpenCommandList()
 		LocalSamplerHeap.OpenCommandList();
 	}
 
-	if (CurrentViewHeap)
+	if (!IsUsingBindlessSamplers())
 	{
-		CurrentViewHeap->OpenCommandList();
+		SwitchToGlobalSamplerHeap();
 	}
 
-	if (!IsUsingBindlessSamplers())
+	if (CurrentViewHeap)
 	{
-		SwitchToGlobalSamplerHeap();
+		CurrentViewHeap->OpenCommandList();
 	}
 
 	// Make sure the heaps are set
