- Add persistent SBT support to d3d12 RHI
- Add validation support for persistent SBT record data

diff --git a/Engine/Source/Runtime/D3D12RHI/Private/D3D12RayTracing.cpp b/Engine/Source/Runtime/D3D12RHI/Private/D3D12RayTracing.cpp
--- a/Engine/Source/Runtime/D3D12RHI/Private/D3D12RayTracing.cpp
+++ b/Engine/Source/Runtime/D3D12RHI/Private/D3D12RayTracing.cpp
@@ -24,6 +24,7 @@
 #include "D3D12ResourceCollection.h"
 #include "RHIShaderBindingLayout.h"
 #include "RHIUniformBufferUtilities.h"
+#include "D3D12TextureReference.h"
 
 extern int32 GD3D12ExplicitViewDescriptorHeapSize;
 extern int32 GD3D12ExplicitViewDescriptorHeapOverflowReported;
@@ -113,6 +114,8 @@ DECLARE_DWORD_COUNTER_STAT(TEXT("Updated BL AS (per frame)"), STAT_D3D12RayTraci
 DECLARE_DWORD_COUNTER_STAT(TEXT("Built TL AS (per frame)"), STAT_D3D12RayTracingBuiltTLAS, STATGROUP_D3D12RayTracing);
 DECLARE_DWORD_COUNTER_STAT(TEXT("Updated TL AS (per frame)"), STAT_D3D12RayTracingUpdatedTLAS, STATGROUP_D3D12RayTracing);
 
+DECLARE_DWORD_COUNTER_STAT(TEXT("UniformBuffer Record Updates (per frame)"), STAT_D3D12RayTracingUBRecordUpdate, STATGROUP_D3D12RayTracing);
+
 DECLARE_MEMORY_STAT(TEXT("Total BL AS Memory"), STAT_D3D12RayTracingBLASMemory, STATGROUP_D3D12RayTracing);
 DECLARE_MEMORY_STAT(TEXT("Static BL AS Memory"), STAT_D3D12RayTracingStaticBLASMemory, STATGROUP_D3D12RayTracing);
 DECLARE_MEMORY_STAT(TEXT("Dynamic BL AS Memory"), STAT_D3D12RayTracingDynamicBLASMemory, STATGROUP_D3D12RayTracing);
@@ -1290,6 +1293,12 @@ private:
 		FMemory::Memcpy(Data.GetData() + WriteOffset, InData, InDataSize);
 	}
 
+	void CompareData(uint32 Offset, const void* InData, uint32 InDataSize)
+	{
+		const void* CurrentData = Data.GetData() + Offset;
+		check(FMemory::Memcmp(CurrentData, InData, InDataSize) == 0);
+	}
+
 	void WriteLocalShaderRecord(uint32 ShaderTableOffset, uint32 RecordIndex, uint32 OffsetWithinRecord, const void* InData, uint32 InDataSize)
 	{
 		checkfSlow(OffsetWithinRecord % 4 == 0, TEXT("SBT record parameters must be written on DWORD-aligned boundary"));
@@ -1301,6 +1310,12 @@ private:
 		WriteData(WriteOffset, InData, InDataSize);
 	}
 
+	void CompareLocalShaderRecord(uint32 ShaderTableOffset, uint32 RecordIndex, uint32 OffsetWithinRecord, const void* InData, uint32 InDataSize)
+	{
+		const uint32 Offset = ShaderTableOffset + LocalRecordStride * RecordIndex + OffsetWithinRecord;
+		CompareData(Offset, InData, InDataSize);
+	}
+
 public:
 
 	UE_NONCOPYABLE(FD3D12RayTracingShaderBindingTableInternal)
@@ -1310,7 +1325,7 @@ public:
 	// Scaling beyond 5 total threads does not yield any speedup in practice.
 	static constexpr uint32 MaxBindingWorkers = 5; // RHI thread + 4 parallel workers.
 
-	FD3D12RayTracingShaderBindingTableInternal(const FRayTracingShaderBindingTableInitializer& Initializer, FD3D12Device* Device)
+	FD3D12RayTracingShaderBindingTableInternal(const FRayTracingShaderBindingTableInitializer& Initializer, FD3D12Device* InDevice)
 		: UniqueId(NextUniqueId++)
 	{
 		checkf(Initializer.LocalBindingDataSize <= 4096, TEXT("The maximum size of a local root signature is 4KB.")); // as per section 4.22.1 of DXR spec v1.0
@@ -1319,6 +1334,9 @@ public:
 		const uint32 NumHitGroupSlots = Initializer.HitGroupIndexingMode == ERayTracingHitGroupIndexingMode::Allow ? Initializer.NumGeometrySegments * Initializer.NumShaderSlotsPerGeometrySegment : 1;
 		checkf(Initializer.LocalBindingDataSize >= sizeof(FD3D12HitGroupSystemParameters), TEXT("All local root signatures are expected to contain ray tracing system root parameters"));
 		
+		Device = InDevice;
+
+		Lifetime = Initializer.Lifetime;
 		HitGroupIndexingMode = Initializer.HitGroupIndexingMode;
 		LocalRecordSizeUnaligned = ShaderIdentifierSize + Initializer.LocalBindingDataSize;
 		LocalRecordStride = RoundUpToNextMultiple(LocalRecordSizeUnaligned, D3D12_RAYTRACING_SHADER_RECORD_BYTE_ALIGNMENT);
@@ -1374,15 +1392,26 @@ public:
 
 		// Keep CPU-side data after upload
 		Data.SetAllowCPUAccess(true);
+
+		if (Lifetime == ERayTracingShaderBindingTableLifetime::Persistent)
+		{
+			HitRecordData.SetNum(NumHitRecords);
+		}
 	}
 
 	~FD3D12RayTracingShaderBindingTableInternal()
 	{
+		// Might be slow - iterate faster somehow?
+		for (uint32 RecordIndex = 0; RecordIndex < NumHitRecords; ++RecordIndex)
+		{
+			ClearHitRecordData(0, RecordIndex);
+		}
+
 		delete DescriptorCache;
 #if D3D12RHI_USE_CONSTANT_BUFFER_VIEWS
 		for (FWorkerThreadData& ThisWorkerData : WorkerData)
 		{
-			for (FD3D12ConstantBufferView* CBV : ThisWorkerData.TransientCBVs)
+			for (FD3D12ConstantBufferView* CBV : ThisWorkerData.LooseParameterCBVs)
 			{
 				delete CBV;
 			}
@@ -1401,6 +1430,17 @@ public:
 		WriteLocalShaderRecord(ShaderTableOffset, RecordIndex, ShaderIdentifierSize + InOffsetWithinRootSignature, InData, InDataSize);
 	}
 
+	template <typename T>
+	void CompareLocalShaderParameters(uint32 ShaderTableOffset, uint32 RecordIndex, uint32 InOffsetWithinRootSignature, const T& Parameters)
+	{
+		CompareLocalShaderRecord(ShaderTableOffset, RecordIndex, ShaderIdentifierSize + InOffsetWithinRootSignature, &Parameters, sizeof(Parameters));
+	}
+
+	void CompareLocalShaderParameters(uint32 ShaderTableOffset, uint32 RecordIndex, uint32 InOffsetWithinRootSignature, const void* InData, uint32 InDataSize)
+	{
+		CompareLocalShaderRecord(ShaderTableOffset, RecordIndex, ShaderIdentifierSize + InOffsetWithinRootSignature, InData, InDataSize);
+	}
+
 	template <typename T>
 	void SetMissShaderParameters(uint32 RecordIndex, uint32 InOffsetWithinRootSignature, const T& Parameters)
 	{
@@ -1466,6 +1506,13 @@ public:
 		SetLocalShaderParameters(ShaderTableOffset, RecordIndex, OffsetWithinRootSignature, SystemParameters);
 	}
 
+	void CompareHitGroupSystemParameters(uint32 RecordIndex, const FD3D12HitGroupSystemParameters& SystemParameters)
+	{
+		const uint32 OffsetWithinRootSignature = 0; // System parameters are always first in the RS.
+		const uint32 ShaderTableOffset = HitGroupShaderTableOffset;
+		CompareLocalShaderParameters(ShaderTableOffset, RecordIndex, OffsetWithinRootSignature, SystemParameters);
+	}
+
 	void SetHitGroupIdentifier(uint32 RecordIndex, const FD3D12ShaderIdentifier& ShaderIdentifier)
 	{
 		checkfSlow(ShaderIdentifier.IsValid(), TEXT("Shader identifier must be initialized FD3D12RayTracingPipelineState::GetShaderIdentifier() before use."));
@@ -1475,6 +1522,15 @@ public:
 		WriteData(WriteOffset, ShaderIdentifier.Data, ShaderIdentifierSize);
 	}	
 
+	void CompareHitGroupIdentifier(uint32 RecordIndex, const FD3D12ShaderIdentifier& ShaderIdentifier)
+	{
+		checkfSlow(ShaderIdentifier.IsValid(), TEXT("Shader identifier must be initialized FD3D12RayTracingPipelineState::GetShaderIdentifier() before use."));
+		checkSlow(sizeof(ShaderIdentifier.Data) >= ShaderIdentifierSize);
+
+		const uint32 WriteOffset = HitGroupShaderTableOffset + RecordIndex * LocalRecordStride;
+		CompareData(WriteOffset, ShaderIdentifier.Data, ShaderIdentifierSize);
+	}
+
 	void SetDefaultMissShaderIdentifier(const FD3D12ShaderIdentifier& ShaderIdentifier)
 	{
 		// Set all slots to the same default
@@ -1508,15 +1564,74 @@ public:
 
 		checkf(bWasDefaultMissShaderSet, TEXT("At least the first miss shader must have been set before copying to GPU."));
 		
+#if ENABLE_RESIDENCY_MANAGEMENT
 		// Merge all data from worker threads into the main set
+		{
+			TRACE_CPUPROFILER_EVENT_SCOPE(MergeWorkerData);
 
-		for (uint32 WorkerIndex = 1; WorkerIndex < MaxBindingWorkers; ++WorkerIndex)
+			for (uint32 WorkerIndex = 1; WorkerIndex < MaxBindingWorkers; ++WorkerIndex)
+			{
+				for (FD3D12Resource* Resource : WorkerData[WorkerIndex].DynamicReferencedResources)
+				{
+					AddDynamicReferencedResource(0, Resource);
+				}
+				for (auto Iter = WorkerData[WorkerIndex].PersistentReferencedShaderResources.CreateIterator(); Iter; ++Iter)
+				{
+					UpdatePersistentReferencedShaderResource(0, Iter.Key(), Iter.Value());
+				}
+				for (auto Iter = WorkerData[WorkerIndex].PersistentReferencedTextureReferences.CreateIterator(); Iter; ++Iter)
+				{
+					UpdatePersistentReferencedTextureReference(0, Iter.Key(), Iter.Value());
+				}
+				WorkerData[WorkerIndex].DynamicReferencedResources.Empty();
+				WorkerData[WorkerIndex].DynamicReferencedResourceSet.Empty();
+				WorkerData[WorkerIndex].PersistentReferencedShaderResources.Empty();
+				WorkerData[WorkerIndex].PersistentReferencedTextureReferences.Empty();
+			}
+		}
+
+		// Build the used persistent d3d12 resource set from the views
 		{
-			for (FD3D12Resource* Resource : WorkerData[WorkerIndex].ReferencedD3D12Resources)
+			TRACE_CPUPROFILER_EVENT_SCOPE(BuildUsedResourceSet);
+
+			// Pretty expensive for a lot of resources - could be kicked in async task or tracking code refactored to
+			// directly track the used d3d12 resources but then SBT needs to be listener to all used base shader resources
+			// for defrag and texture reference changes
+			TSet<FD3D12Resource*> ReferencedResourcesSet;
+			ReferencedResourcesSet.Reserve(ReferencedResources.Num());
+			for (auto Iter = WorkerData[0].PersistentReferencedShaderResources.CreateIterator(); Iter; ++Iter)
+			{
+				FD3D12BaseShaderResource* ShaderResource = Iter.Key();
+				// buffer could have been streamed out but MDC is still valid containing the RHI buffer due to internal swap
+				if (ShaderResource->GetResource())
+				{
+					ReferencedResourcesSet.Add(ShaderResource->GetResource());
+				}
+			}
+			for (auto Iter = WorkerData[0].PersistentReferencedTextureReferences.CreateIterator(); Iter; ++Iter)
 			{
-				AddResourceReference(Resource, 0);
+				FD3D12RHITextureReference* TextureReference = Iter.Key();
+				FD3D12Texture* Texture = FD3D12CommandContext::RetrieveTexture(TextureReference, Device->GetGPUIndex());
+				check(Texture->GetResource());
+				ReferencedResourcesSet.Add(Texture->GetResource());
 			}
 
+			ReferencedResources.Empty(ReferencedResourcesSet.Num() + WorkerData[0].DynamicReferencedResources.Num());
+			for (TSet<FD3D12Resource*>::TConstIterator SetIt(ReferencedResourcesSet); SetIt; ++SetIt)
+			{
+				ReferencedResources.Add(*SetIt);
+			}
+
+			// Also add all the collected dynamic resources and reset for the next frame			
+			ReferencedResources.Append(WorkerData[0].DynamicReferencedResources);
+			WorkerData[0].DynamicReferencedResources.Empty();
+			WorkerData[0].DynamicReferencedResourceSet.Empty();
+		}
+#endif // ENABLE_RESIDENCY_MANAGEMENT
+
+		// Merge all transitions and perform other non residency related per worker thread
+		for (uint32 WorkerIndex = 1; WorkerIndex < MaxBindingWorkers; ++WorkerIndex)
+		{
 			for (FD3D12ShaderResourceView* SRV : WorkerData[WorkerIndex].TransitionSRVs)
 			{
 				AddResourceTransition(SRV, 0);
@@ -1527,13 +1642,14 @@ public:
 				AddResourceTransition(UAV, 0);
 			}
 
-			WorkerData[WorkerIndex].ReferencedD3D12Resources.Empty();
 			WorkerData[WorkerIndex].TransitionSRVs.Empty();
 			WorkerData[WorkerIndex].TransitionUAVs.Empty();
 			WorkerData[WorkerIndex].TransitionViewSet.Empty();
+
+			// Clear the record cache for reuse
+			WorkerData[WorkerIndex].ShaderRecordCache.Empty();
 		}
 
-		FD3D12Device* Device = Context.GetParentDevice();
 		FD3D12Adapter* Adapter = Device->GetParentAdapter();
 		D3D12_RESOURCE_DESC BufferDesc = CD3DX12_RESOURCE_DESC::Buffer(Data.GetResourceDataSize(), D3D12_RESOURCE_FLAG_NONE, D3D12_RAYTRACING_SHADER_TABLE_BYTE_ALIGNMENT);
 		FRHIGPUMask GPUMask = FRHIGPUMask::FromIndex(Device->GetGPUIndex());
@@ -1563,8 +1679,10 @@ public:
 		return Buffer->ResourceLocation.GetGPUVirtualAddress();
 	}
 
-	D3D12_DISPATCH_RAYS_DESC GetDispatchRaysDesc(FD3D12Device* Device, const FD3D12ShaderIdentifier& RayGenShaderIdentifier) const
+	D3D12_DISPATCH_RAYS_DESC GetDispatchRaysDesc(FD3D12Device* InDevice, const FD3D12ShaderIdentifier& RayGenShaderIdentifier) const
 	{
+		check(Device == InDevice);
+
 		// Allocate memory for the ray gen shader identifier storage
 		check(ShaderIdentifierSize == sizeof(FD3D12ShaderIdentifier));
 		FD3D12ResourceLocation UploadResourceLocation(Device);
@@ -1616,6 +1734,9 @@ public:
 	uint32 HitGroupShaderTableOffset = 0;
 	uint32 CallableShaderTableOffset = 0;
 
+	FD3D12Device* Device  = nullptr;
+
+	ERayTracingShaderBindingTableLifetime Lifetime = ERayTracingShaderBindingTableLifetime::Transient;
 	ERayTracingHitGroupIndexingMode HitGroupIndexingMode = ERayTracingHitGroupIndexingMode::Allow;
 
 	// Note: TABLE_BYTE_ALIGNMENT is used instead of RECORD_BYTE_ALIGNMENT to allow arbitrary switching 
@@ -1632,9 +1753,429 @@ public:
 	bool bWasDefaultMissShaderSet = false;
 #endif
 
+	// All reference d3d12 resources from both persistent and dynamic records (rebuild during commit)
+	TArray<FD3D12Resource*> ReferencedResources;
+	
 	// SBTs have their own descriptor heaps
 	FD3D12ExplicitDescriptorCache* DescriptorCache = nullptr;
 
+	struct FRecordUpdateUniformBufferListener : public ID3D12UniformBufferUpdateListener
+	{
+		FD3D12RayTracingShaderBindingTableInternal& ShaderTable;
+		FD3D12UniformBuffer* UniformBuffer = nullptr;
+		uint32 ShaderTableOffset = 0;
+		uint32 RecordIndex = 0;
+		uint32 OffsetWithinRootSignature = 0;
+		
+		FRecordUpdateUniformBufferListener(FD3D12RayTracingShaderBindingTableInternal& InShaderTable) :	ShaderTable(InShaderTable)
+		{
+		}
+
+		void Init(FD3D12UniformBuffer* InUniformBuffer, uint32 InShaderTableOffset, uint32 InRecordIndex, uint32 InOffsetWithinRootSignature)
+		{
+			UniformBuffer = InUniformBuffer;
+			ShaderTableOffset = InShaderTableOffset;
+			RecordIndex = InRecordIndex;
+			OffsetWithinRootSignature = InOffsetWithinRootSignature;
+
+			UniformBuffer->AddUpdateListener(this);
+		}
+
+		void Reset()
+		{
+			check(UniformBuffer);
+			UniformBuffer->RemoveUpdateListener(this);
+			UniformBuffer = nullptr;
+		}
+
+		virtual void RemoveListener(FD3D12UniformBuffer* InUpdatedUniformBuffer) final override
+		{
+			check(UniformBuffer == InUpdatedUniformBuffer);
+			ShaderTable.RemoveUpdateUniformBufferListener(RecordIndex, this);
+		}
+
+		virtual void UniformBufferUpdated(FRHICommandListBase& CmdList, FD3D12UniformBuffer* InUpdatedUniformBuffer) final override
+		{
+			check(UniformBuffer == InUpdatedUniformBuffer);
+			ShaderTable.SetLocalShaderParameters(ShaderTableOffset, RecordIndex, OffsetWithinRootSignature, InUpdatedUniformBuffer->ResourceLocation.GetGPUVirtualAddress());
+
+			INC_DWORD_STAT(STAT_D3D12RayTracingUBRecordUpdate);
+		}
+	};
+
+	struct FRecordUpdateRayTracingGeometryListener : public ID3D12RayTracingGeometryUpdateListener
+	{
+		FD3D12RayTracingShaderBindingTableInternal* ShaderTable = nullptr;
+		const FD3D12RayTracingGeometry* Geometry = nullptr;
+		uint32 RecordIndex = 0;
+		uint32 UserData = 0;
+		uint32 SegmentIndex = 0;
+
+		void Init(FD3D12RayTracingShaderBindingTableInternal* InShaderTable, const FD3D12RayTracingGeometry* InGeometry, uint32 InRecordIndex, uint32 InUserData, uint32 InSegmentIndex)
+		{
+			check(Geometry == nullptr);
+
+			ShaderTable = InShaderTable;
+			Geometry = InGeometry;
+			RecordIndex = InRecordIndex;
+			UserData = InUserData;
+			SegmentIndex = InSegmentIndex;
+
+			Geometry->AddUpdateListener(this);
+		}
+
+		void Reset()
+		{
+			ShaderTable = nullptr;
+			Geometry = nullptr;
+			RecordIndex = 0;
+			UserData = 0;
+			SegmentIndex = 0;
+		}
+
+		virtual void RemoveListener(FD3D12RayTracingGeometry* InGeometry) final override
+		{
+			check(Geometry == InGeometry);
+			ShaderTable->RemoveUpdateRayTracingGeometryListener(RecordIndex, this);
+		}
+
+		virtual void HitGroupParametersUpdated(FD3D12RayTracingGeometry* InGeometry) final override
+		{
+			check(InGeometry == Geometry);
+
+			uint32 GPUIndex = ShaderTable->Device->GetGPUIndex();
+			FD3D12HitGroupSystemParameters SystemParameters = Geometry->HitGroupSystemParameters[GPUIndex][SegmentIndex];
+			SystemParameters.RootConstants.UserData = UserData;
+			ShaderTable->SetHitGroupSystemParameters(RecordIndex, SystemParameters);
+		}
+	};
+
+	// Tracked data for each record so stored GPU data can be updated after defrag operation
+	struct FRecordData
+	{
+		enum class EState
+		{
+			Uninitialized,
+			Persistent,
+			Transient,
+		};
+
+		void Reset()
+		{
+#if ENABLE_RESIDENCY_MANAGEMENT
+			ReferencedShaderResources.Empty();
+			ReferencedTextureReferences.Empty();
+#endif // ENABLE_RESIDENCY_MANAGEMENT
+			UpdateGeometryListener.Reset();
+			UpdateUniformBufferListeners.Empty();
+#if DO_CHECK
+			LooseParameterData.Empty();
+#endif //DO_CHECK
+
+			State = EState::Uninitialized;
+		}
+
+		EState State = EState::Uninitialized;
+		TArray<FRecordUpdateUniformBufferListener*> UpdateUniformBufferListeners;
+		FRecordUpdateRayTracingGeometryListener UpdateGeometryListener;
+		FD3D12ResourceLocation* LooseParameterResourceLocation = nullptr;
+
+#if ENABLE_RESIDENCY_MANAGEMENT		
+		TArray<FD3D12BaseShaderResource*> ReferencedShaderResources;
+		TArray<FD3D12RHITextureReference*> ReferencedTextureReferences;
+#endif // ENABLE_RESIDENCY_MANAGEMENT		
+#if DO_CHECK
+		TArray<uint8> LooseParameterData;
+#endif //DO_CHECK
+	};
+	TArray<FRecordData> HitRecordData;
+
+	void SetHitGroupSystemParameters(uint32 WorkerIndex, ERayTracingLocalShaderBindingType BindingType, uint32 RecordIndex, const FD3D12RayTracingGeometry* Geometry, uint32 UserData, uint32 SegmentIndex)
+	{	
+		FD3D12HitGroupSystemParameters SystemParameters = Geometry->HitGroupSystemParameters[Device->GetGPUIndex()][SegmentIndex];
+		SystemParameters.RootConstants.UserData = UserData;
+
+		if (Lifetime == ERayTracingShaderBindingTableLifetime::Transient)
+		{
+			check(BindingType == ERayTracingLocalShaderBindingType::Transient)
+			SetHitGroupSystemParameters(RecordIndex, SystemParameters);
+		}
+		else
+		{
+			FRecordData& RecordData = HitRecordData[RecordIndex];
+			bool bSetData = BindingType != ERayTracingLocalShaderBindingType::Validation;
+
+			// Special use case when trying to set transient binding on already persistent stored binding (validate content only then)
+			// Can happen we go from persistent -> transient -> persistent when RTPSO is still compiling for example
+			if (RecordData.State == FRecordData::EState::Persistent && BindingType == ERayTracingLocalShaderBindingType::Transient)
+			{
+				bSetData = false;
+			}
+
+			if (bSetData)
+			{
+				check(RecordData.State != FRecordData::EState::Persistent || BindingType == ERayTracingLocalShaderBindingType::Transient);
+				SetHitGroupSystemParameters(RecordIndex, SystemParameters);
+
+				// If persistently stored then also listen to hit record updates for stream in/out events of the ray tracing geometry data
+				if (BindingType == ERayTracingLocalShaderBindingType::Persistent)
+				{
+					HitRecordData[RecordIndex].UpdateGeometryListener.Init(this, Geometry, RecordIndex, UserData, SegmentIndex);
+					RecordData.State = FRecordData::EState::Persistent;
+				}				
+				else // if (RecordData.State == FRecordData::EState::Uninitialized)
+				{
+					RecordData.State = FRecordData::EState::Transient;
+				}
+			}
+			else
+			{
+				check(RecordData.State == FRecordData::EState::Persistent);
+				CompareHitGroupSystemParameters(RecordIndex, SystemParameters);
+			}
+		}
+	}
+		
+#if ENABLE_RESIDENCY_MANAGEMENT	
+	void UpdatePersistentReferencedShaderResource(uint32 WorkerIndex, FD3D12BaseShaderResource* ShaderResource, int32 Count)
+	{
+		// Update the worker thread reference count of the resource
+		int32& ReferenceCount = WorkerData[WorkerIndex].PersistentReferencedShaderResources.FindOrAdd(ShaderResource);
+		ReferenceCount += Count;
+		
+		if (WorkerIndex == 0)
+		{
+			check(ReferenceCount >= 0);
+			if (ReferenceCount == 0)
+			{
+				WorkerData[WorkerIndex].PersistentReferencedShaderResources.Remove(ShaderResource);
+			}
+		}
+	}
+
+	void UpdatePersistentReferencedTextureReference(uint32 WorkerIndex, FD3D12RHITextureReference* TextureReference, int32 Count)
+	{
+		// Update the worker thread reference count of the resource
+		int32& ReferenceCount = WorkerData[WorkerIndex].PersistentReferencedTextureReferences.FindOrAdd(TextureReference);
+		ReferenceCount += Count;
+		
+		if (WorkerIndex == 0)
+		{
+			check(ReferenceCount >= 0);
+			if (ReferenceCount == 0)
+			{
+				WorkerData[WorkerIndex].PersistentReferencedTextureReferences.Remove(TextureReference);
+			}
+		}
+	}
+
+	void AddDynamicReferencedResource(uint32 WorkerIndex, FD3D12Resource* Resource)
+	{
+		bool bIsAlreadyInSet = false;
+		WorkerData[WorkerIndex].DynamicReferencedResourceSet.Add(Resource, &bIsAlreadyInSet);
+		if (!bIsAlreadyInSet)
+		{
+			WorkerData[WorkerIndex].DynamicReferencedResources.Add(Resource);
+		}
+	}
+#endif // ENABLE_RESIDENCY_MANAGEMENT
+
+	void AddReferencedShaderResource(uint32 WorkerIndex, uint32 RecordIndex, ERayTracingLocalShaderBindingType BindingType, FD3D12BaseShaderResource* ShaderResource)
+	{
+#if ENABLE_RESIDENCY_MANAGEMENT
+		if (BindingType == ERayTracingLocalShaderBindingType::Transient || Lifetime != ERayTracingShaderBindingTableLifetime::Persistent)
+		{
+			AddDynamicReferencedResource(WorkerIndex, ShaderResource->GetResource());
+		}
+		else if (BindingType == ERayTracingLocalShaderBindingType::Persistent)
+		{
+			// Update the tracked record data so it can be updated in case of defrag operation
+			FRecordData& RecordData = HitRecordData[RecordIndex];
+			RecordData.ReferencedShaderResources.Add(ShaderResource);
+			UpdatePersistentReferencedShaderResource(WorkerIndex, ShaderResource, 1);
+		}
+#endif // ENABLE_RESIDENCY_MANAGEMENT
+	}
+
+	void AddReferencedTexture(uint32 WorkerIndex, uint32 RecordIndex, ERayTracingLocalShaderBindingType BindingType, FRHITexture* RHITexture)
+	{
+#if ENABLE_RESIDENCY_MANAGEMENT
+		if (BindingType == ERayTracingLocalShaderBindingType::Transient || Lifetime != ERayTracingShaderBindingTableLifetime::Persistent)
+		{
+			FD3D12Texture* Texture = FD3D12CommandContext::RetrieveTexture(RHITexture, Device->GetGPUIndex());			
+			AddDynamicReferencedResource(WorkerIndex, Texture->ResourceLocation.GetResource());
+		}
+		else if (BindingType == ERayTracingLocalShaderBindingType::Persistent)
+		{
+			// Update the tracked record data so it can be updated in case of texture reference swap
+			FRecordData& RecordData = HitRecordData[RecordIndex];
+			
+			// Collect all texture references independent of the shader resources because the internal shader resource can be changed on swap
+			if (FRHITextureReference* TextureReferenceRHI = RHITexture->GetTextureReference())
+			{
+				FD3D12RHITextureReference* TextureReference = FD3D12CommandContext::RetrieveObject<FD3D12RHITextureReference>(TextureReferenceRHI, Device->GetGPUIndex());
+				RecordData.ReferencedTextureReferences.Add(TextureReference);
+				UpdatePersistentReferencedTextureReference(WorkerIndex, TextureReference, 1);
+			}
+			else
+			{
+				FD3D12Texture* Texture = FD3D12CommandContext::RetrieveTexture(RHITexture, Device->GetGPUIndex());
+				RecordData.ReferencedShaderResources.Add(Texture);
+				UpdatePersistentReferencedShaderResource(WorkerIndex, Texture, 1);
+			}
+		}
+#endif // ENABLE_RESIDENCY_MANAGEMENT
+	}
+
+	void AddReferencedUniformBuffer(uint32 WorkerIndex, uint32 RecordIndex, uint32 ShaderTableOffset, uint32 InOffsetWithinRootSignature, ERayTracingLocalShaderBindingType BindingType, FD3D12UniformBuffer* UniformBuffer)
+	{
+		// Uniform buffer shouldn't have an residency handles because they are all allocated in upload memory so doesn't require extra tracking
+		check(UniformBuffer->ResourceLocation.GetResource()->GetResidencyHandles().IsEmpty());
+
+		if (Lifetime == ERayTracingShaderBindingTableLifetime::Persistent && BindingType == ERayTracingLocalShaderBindingType::Persistent)
+		{
+			FRecordData& RecordData = HitRecordData[RecordIndex];
+
+			for (FRecordUpdateUniformBufferListener* Listener : RecordData.UpdateUniformBufferListeners)
+			{
+				check(Listener->UniformBuffer != UniformBuffer);
+			}
+			
+			FRecordUpdateUniformBufferListener* Listener = new FRecordUpdateUniformBufferListener(*this);
+			Listener->Init(UniformBuffer, ShaderTableOffset, RecordIndex, InOffsetWithinRootSignature);
+			RecordData.UpdateUniformBufferListeners.Add(Listener);
+		}
+	}
+
+	FD3D12ConstantBufferView* SetLooseParameterData(
+		uint32 WorkerIndex, 
+		uint32 RecordIndex, 
+		ERayTracingLocalShaderBindingType BindingType, 
+		const void* LooseParameterData, 
+		uint32 LooseParameterDataSize, 
+		D3D12_GPU_VIRTUAL_ADDRESS& OutGPUVirtualAddress)
+	{
+		// If we see a significant number of LooseParameter allocations coming through this path, we should consider
+		// caching constant buffer blocks inside ShaderTable and linearly sub-allocate from them.
+		// If the amount of data is relatively small, it may also be possible to use root constants and avoid extra allocations entirely.
+
+		FD3D12ConstantBufferView* ConstantBufferView = nullptr;
+
+		if (BindingType == ERayTracingLocalShaderBindingType::Validation)
+		{
+#if DO_CHECK
+			FRecordData& RecordData = HitRecordData[RecordIndex];
+			check(RecordData.LooseParameterData.Num() == LooseParameterDataSize);
+			check(FMemory::Memcmp(RecordData.LooseParameterData.GetData(), LooseParameterData, LooseParameterDataSize) == 0);
+#endif // DO_CHECK
+		}
+		else
+		{
+			// Always allocate using the upload heap allocator and not the TransientUniformBufferAllocator (can cause GPU crashes with transient data somehow)
+			FD3D12ResourceLocation ResourceLocation(Device);
+			FD3D12UploadHeapAllocator& Allocator = Device->GetParentAdapter()->GetUploadHeapAllocator(Device->GetGPUIndex());
+			void* MappedData = Allocator.AllocUploadResource(LooseParameterDataSize, D3D12_CONSTANT_BUFFER_DATA_PLACEMENT_ALIGNMENT, ResourceLocation);
+			FMemory::Memcpy(MappedData, LooseParameterData, LooseParameterDataSize);
+			OutGPUVirtualAddress = ResourceLocation.GetGPUVirtualAddress();
+
+			// Uniform buffer shouldn't have an residency handles because they are all allocated in upload memory so doesn't require residency tracking
+			check(ResourceLocation.GetResource()->GetResidencyHandles().IsEmpty());
+
+#if D3D12RHI_USE_CONSTANT_BUFFER_VIEWS
+			// Not implemented for persistent SBT yet (CBV needs to be managed per record then instead of at SBT level)
+			// 
+			// CBV also has pointer to stack local ResourceLocation which it shouldn't touch anymore after creation but
+			// will still have dangling pointer - might be better to always allocate ResourceLocation as well on heap next to CBV (like persistent path)
+			check(Lifetime == ERayTracingShaderBindingTableLifetime::Transient);
+			ConstantBufferView = new FD3D12ConstantBufferView(Device, nullptr);
+			ConstantBufferView->CreateView(&ResourceLocation, 0, Align(LooseParameterDataSize, D3D12_CONSTANT_BUFFER_DATA_PLACEMENT_ALIGNMENT));
+			WorkerData[WorkerIndex].LooseParameterCBVs.Add(ConstantBufferView);			
+#endif // D3D12RHI_USE_CONSTANT_BUFFER_VIEWS
+
+			// Move ownership of the upload memory to the record data when persistent (otherwise auto freed next 'frame')
+			if (Lifetime == ERayTracingShaderBindingTableLifetime::Persistent && BindingType == ERayTracingLocalShaderBindingType::Persistent)
+			{
+				FRecordData& RecordData = HitRecordData[RecordIndex];
+				check(RecordData.LooseParameterResourceLocation == nullptr);
+				RecordData.LooseParameterResourceLocation = new FD3D12ResourceLocation(Device);
+				FD3D12ResourceLocation::TransferOwnership(*RecordData.LooseParameterResourceLocation, ResourceLocation);
+
+#if DO_CHECK
+				RecordData.LooseParameterData.SetNum(LooseParameterDataSize);
+				FMemory::Memcpy(RecordData.LooseParameterData.GetData(), LooseParameterData, LooseParameterDataSize);
+#endif // DO_CHECK
+			}
+		}
+
+		return ConstantBufferView;
+	}
+
+	void ClearHitRecordData(uint32 WorkerIndex, uint32 RecordIndex)
+	{
+		if (Lifetime != ERayTracingShaderBindingTableLifetime::Persistent)
+		{
+			return;
+		}
+
+		FRecordData& RecordData = HitRecordData[RecordIndex];
+		
+		for (FRecordUpdateUniformBufferListener* Listener : RecordData.UpdateUniformBufferListeners)
+		{
+			Listener->Reset();
+			delete Listener;
+		}
+
+		if (RecordData.LooseParameterResourceLocation)
+		{
+			delete RecordData.LooseParameterResourceLocation;
+			RecordData.LooseParameterResourceLocation = nullptr;
+		}
+
+		if (RecordData.UpdateGeometryListener.Geometry)
+		{
+			RecordData.UpdateGeometryListener.Geometry->RemoveUpdateListener(&RecordData.UpdateGeometryListener);
+		}
+		
+#if ENABLE_RESIDENCY_MANAGEMENT
+		for (FD3D12BaseShaderResource* ShaderResource : RecordData.ReferencedShaderResources)
+		{
+			UpdatePersistentReferencedShaderResource(WorkerIndex, ShaderResource, -1);
+		}
+
+		for (FD3D12RHITextureReference* TextureReference : RecordData.ReferencedTextureReferences)
+		{
+			UpdatePersistentReferencedTextureReference(WorkerIndex, TextureReference, -1);
+		}
+#endif // ENABLE_RESIDENCY_MANAGEMENT
+
+		// Reset record to default state
+		RecordData.Reset();
+	}
+
+	void RemoveUpdateUniformBufferListener(uint32 RecordIndex, FRecordUpdateUniformBufferListener* Listener)
+	{
+		bool bFound = false;
+		for (int32 Index = 0; Index < HitRecordData[RecordIndex].UpdateUniformBufferListeners.Num(); ++Index)
+		{
+			if (HitRecordData[RecordIndex].UpdateUniformBufferListeners[Index] == Listener)
+			{
+				HitRecordData[RecordIndex].UpdateUniformBufferListeners.RemoveAt(Index);
+				delete Listener;
+
+				bFound = true;
+				break;
+			}
+		}
+		check(bFound);
+	}
+
+	void RemoveUpdateRayTracingGeometryListener(uint32 RecordIndex, FRecordUpdateRayTracingGeometryListener* Listener)
+	{
+		FRecordData& RecordData = HitRecordData[RecordIndex];
+		check(&RecordData.UpdateGeometryListener == Listener && RecordData.UpdateGeometryListener.Geometry != nullptr);
+		RecordData.UpdateGeometryListener.Reset();
+	}
+
 	struct FShaderRecordCacheKey
 	{
 		static constexpr uint32 MaxUniformBuffers = 6;
@@ -1676,16 +2217,6 @@ public:
 		}
 	};
 	
-	void AddResourceReference(FD3D12Resource* D3D12Resource, uint32 WorkerIndex)
-	{
-		bool bIsAlreadyInSet = false;
-		WorkerData[WorkerIndex].ReferencedD3D12ResourceSet.Add(D3D12Resource, &bIsAlreadyInSet);
-		if (!bIsAlreadyInSet)
-		{
-			WorkerData[WorkerIndex].ReferencedD3D12Resources.Add(D3D12Resource);
-		}
-	}
-
 	void UpdateResidency(FD3D12CommandContext& CommandContext) const
 	{
 		// Skip redundant resource residency updates when a shader table is repeatedly used on the same command list
@@ -1696,9 +2227,9 @@ public:
 			return;
 		}
 
-		// Use the main (merged) set data to actually update resource residency
+		TRACE_CPUPROFILER_EVENT_SCOPE(ShaderTableUpdateResidency);
 
-		for (FD3D12Resource* Resource : WorkerData[0].ReferencedD3D12Resources)
+		for (FD3D12Resource* Resource : ReferencedResources)
 		{
 			CommandContext.UpdateResidency(Resource);
 		}
@@ -1744,18 +2275,24 @@ public:
 	{
 		Experimental::TSherwoodMap<FShaderRecordCacheKey, uint32> ShaderRecordCache;
 
-		// A set of all resources referenced by this shader table for the purpose of updating residency before ray tracing work dispatch.
-		Experimental::TSherwoodSet<void*> ReferencedD3D12ResourceSet;
-		TArray<TRefCountPtr<FD3D12Resource>> ReferencedD3D12Resources;
-
 		// Some resources referenced in SBT may be dynamic (written on GPU timeline) and may require transition barriers.
 		// We save such resources while we fill the SBT and issue transitions before the SBT is used.
 		Experimental::TSherwoodSet<FD3D12View*> TransitionViewSet;
 		TArray<FD3D12ShaderResourceView*> TransitionSRVs;
 		TArray<FD3D12UnorderedAccessView*> TransitionUAVs;
 
+#if ENABLE_RESIDENCY_MANAGEMENT
+		// Track all persistent resources with number of refs accross all records (when ref reaches 0 then it doesn't need be tracked anymore)
+		TMap<TRefCountPtr<FD3D12BaseShaderResource>, int32> PersistentReferencedShaderResources;
+		TMap<TRefCountPtr<FD3D12RHITextureReference>, int32> PersistentReferencedTextureReferences;
+
+		// A set of all resources referenced by this shader table for the purpose of updating residency before ray tracing work dispatch.
+		Experimental::TSherwoodSet<void*> DynamicReferencedResourceSet;
+		TArray<TRefCountPtr<FD3D12Resource>> DynamicReferencedResources;
+#endif // ENABLE_RESIDENCY_MANAGEMENT
+
 #if D3D12RHI_USE_CONSTANT_BUFFER_VIEWS
-		TArray<FD3D12ConstantBufferView*> TransientCBVs;
+		TArray<FD3D12ConstantBufferView*> LooseParameterCBVs;
 #endif // D3D12RHI_USE_CONSTANT_BUFFER_VIEWS
 	};
 
@@ -3079,6 +3616,14 @@ void FD3D12RayTracingGeometry::ReleaseUnderlyingResource()
 
 FD3D12RayTracingGeometry::~FD3D12RayTracingGeometry()
 {
+	// RT geometry can be destroyed before persistent SBT records are cleared from the cached MDCs 
+	// because they are still pending removal from the scene primitives
+	for (ID3D12RayTracingGeometryUpdateListener* UpdateListener : UpdateListeners)
+	{
+		UpdateListener->RemoveListener(this);
+	}
+	UpdateListeners.Empty();
+
 	ReleaseUnderlyingResource();
 }
 
@@ -3354,6 +3899,9 @@ void FD3D12RayTracingGeometry::SetupHitGroupSystemParameters(uint32 InGPUIndex)
 
 		HitGroupSystemParametersForThisGPU.Add(SystemParameters);
 	}
+
+	// Notify listeners about changes
+	HitGroupParamatersUpdated();
 }
 
 void FD3D12RayTracingGeometry::CreateAccelerationStructureBuildDesc(FD3D12CommandContext& CommandContext, EAccelerationStructureBuildMode BuildMode, D3D12_GPU_VIRTUAL_ADDRESS ScratchBufferAddress, D3D12_BUILD_RAYTRACING_ACCELERATION_STRUCTURE_DESC& OutDesc, TArrayView<D3D12_RAYTRACING_GEOMETRY_DESC>& OutGeometryDescs) const
@@ -4190,6 +4738,8 @@ void FD3D12CommandContext::RHIClearShaderBindingTable(FRHIShaderBindingTable* In
 	SBT->ReleaseForDevice(GetParentDevice());
 }
 
+static constexpr uint32 LooseParameterCBVIndex = 0; // Global uniform buffer is always assumed to be in slot 0
+
 struct FD3D12RayTracingGlobalResourceBinder
 {
 	FD3D12RayTracingGlobalResourceBinder(FD3D12CommandContext& InCommandContext, FD3D12ExplicitDescriptorCache& InDescriptorCache)
@@ -4198,6 +4748,11 @@ struct FD3D12RayTracingGlobalResourceBinder
 	{
 	}
 
+	void SetUniformBuffer(uint32 BaseSlotIndex, uint32 DescriptorIndex, FD3D12UniformBuffer* UniformBuffer)
+	{
+		// don't have to do anything
+	}
+
 	void SetRootCBV(uint32 BaseSlotIndex, uint32 DescriptorIndex, D3D12_GPU_VIRTUAL_ADDRESS Address)
 	{
 		CommandContext.GraphicsCommandList()->SetComputeRootConstantBufferView(BaseSlotIndex + DescriptorIndex, Address);
@@ -4213,15 +4768,26 @@ struct FD3D12RayTracingGlobalResourceBinder
 		CommandContext.GraphicsCommandList()->SetComputeRootDescriptorTable(SlotIndex, DescriptorTable);
 	}
 
-	FD3D12ConstantBufferView* CreateTransientConstantBuffer(FD3D12ResourceLocation& ResourceLocation, const void* Data, uint32 DataSize)
+	FD3D12ConstantBufferView* SetLooseParameterData(const void* Data, uint32 DataSize, D3D12_GPU_VIRTUAL_ADDRESS& OutGPUVirtualAddress)
 	{
-		checkf(0, TEXT("Loose parameters and transient constant buffers are not implemented for global ray tracing shaders (raygen, miss, callable)"));
+		checkf(0, TEXT("Loose parameters are not implemented for global ray tracing shaders (raygen, miss, callable)"));
 		return nullptr;
 	}
 
-	void AddResourceReference(FD3D12Resource* D3D12Resource)
+	void AddReferencedShaderResource(FD3D12BaseShaderResource* ShaderResource)
 	{
-		CommandContext.UpdateResidency(D3D12Resource);
+		CommandContext.UpdateResidency(ShaderResource->GetResource());
+	}
+
+	void AddReferencedTexture(FRHITexture* RHITexture)
+	{
+		FD3D12Texture* Texture = FD3D12CommandContext::RetrieveTexture(RHITexture, CommandContext.GetGPUIndex());
+		CommandContext.UpdateResidency(Texture->ResourceLocation.GetResource());
+	}
+	
+	void AddReferencedUniformBuffer(uint32 BaseSlotIndex, uint32 DescriptorIndex, FD3D12UniformBuffer* UniformBuffer)
+	{
+		CommandContext.UpdateResidency(UniformBuffer->ResourceLocation.GetResource());
 	}
 
 	void AddResourceTransition(FD3D12ShaderResourceView* SRV)
@@ -4258,11 +4824,19 @@ struct FD3D12RayTracingGlobalResourceBinder
 
 struct FD3D12RayTracingLocalResourceBinder
 {
-	FD3D12RayTracingLocalResourceBinder(FD3D12Device& InDevice, FD3D12RayTracingShaderBindingTableInternal& InShaderTable, const FD3D12RootSignature& InRootSignature, uint32 InRecordIndex, uint32 InWorkerIndex, ERayTracingBindingType InBindingType)
+	FD3D12RayTracingLocalResourceBinder(
+		FD3D12Device& InDevice, 
+		FD3D12RayTracingShaderBindingTableInternal& InShaderTable, 
+		const FD3D12RootSignature& InRootSignature, 
+		ERayTracingLocalShaderBindingType InBindingType,
+		uint32 InRecordIndex, 
+		uint32 InWorkerIndex, 
+		ERayTracingBindingType RTBindingType)
 		: Device(InDevice)
 		, ShaderTable(InShaderTable)
 		, DescriptorCache(*InShaderTable.DescriptorCache)
 		, RootSignature(InRootSignature)
+		, BindingType(InBindingType)
 		, RecordIndex(InRecordIndex)
 		, WorkerIndex(InWorkerIndex)
 	{
@@ -4271,7 +4845,7 @@ struct FD3D12RayTracingLocalResourceBinder
 		check(WorkerIndex < uint32(DescriptorCache.WorkerData.Num()));
 		check(RecordIndex != ~0u);
 
-		switch (InBindingType)
+		switch (RTBindingType)
 		{
 		case ERayTracingBindingType::CallableShader:
 			ShaderTableOffset = InShaderTable.CallableShaderTableOffset;
@@ -4287,16 +4861,36 @@ struct FD3D12RayTracingLocalResourceBinder
 		}
 	}
 
-	void SetRootDescriptor(uint32 BaseSlotIndex, uint32 DescriptorIndex, D3D12_GPU_VIRTUAL_ADDRESS Address)
+	uint32 ComputeOffsetWithinRootSignature(uint32 BaseSlotIndex, uint32 DescriptorIndex)
 	{
 		const uint32 BindOffsetBase = RootSignature.GetBindSlotOffsetInBytes(BaseSlotIndex);
 		const uint32 DescriptorSize = uint32(sizeof(D3D12_GPU_VIRTUAL_ADDRESS));
-		const uint32 CurrentOffset = BindOffsetBase + DescriptorIndex * DescriptorSize;
-		ShaderTable.SetLocalShaderParameters(ShaderTableOffset, RecordIndex, CurrentOffset, Address);
+		return BindOffsetBase + DescriptorIndex * DescriptorSize;
+	}
+
+	void SetRootDescriptor(uint32 BaseSlotIndex, uint32 DescriptorIndex, D3D12_GPU_VIRTUAL_ADDRESS Address)
+	{
+		const uint32 OffsetWithinRootSignature = ComputeOffsetWithinRootSignature(BaseSlotIndex, DescriptorIndex);
+		if (BindingType == ERayTracingLocalShaderBindingType::Validation)
+		{
+			ShaderTable.CompareLocalShaderParameters(ShaderTableOffset, RecordIndex, OffsetWithinRootSignature, Address);
+		}
+		else
+		{
+			ShaderTable.SetLocalShaderParameters(ShaderTableOffset, RecordIndex, OffsetWithinRootSignature, Address);
+		}
 	}
 
 	void SetRootCBV(uint32 BaseSlotIndex, uint32 DescriptorIndex, D3D12_GPU_VIRTUAL_ADDRESS Address)
 	{
+#if DO_CHECK
+		// Skip validation for descriptor index 0 when loose parameter parameter data is available because that's
+		// always stored at index 0 and validated at loose parameter set time
+		if (DescriptorIndex == LooseParameterCBVIndex && BindingType == ERayTracingLocalShaderBindingType::Validation && bHasLooseParameterData)
+		{
+			return;
+		}
+#endif // DO_CHECK
 		SetRootDescriptor(BaseSlotIndex, DescriptorIndex, Address);
 	}
 
@@ -4308,35 +4902,36 @@ struct FD3D12RayTracingLocalResourceBinder
 	void SetRootDescriptorTable(uint32 SlotIndex, D3D12_GPU_DESCRIPTOR_HANDLE DescriptorTable)
 	{
 		const uint32 BindOffset = RootSignature.GetBindSlotOffsetInBytes(SlotIndex);
-		ShaderTable.SetLocalShaderParameters(ShaderTableOffset, RecordIndex, BindOffset, DescriptorTable);
+		if (BindingType == ERayTracingLocalShaderBindingType::Validation)
+		{
+			ShaderTable.CompareLocalShaderParameters(ShaderTableOffset, RecordIndex, BindOffset, DescriptorTable);
+		}
+		else
+		{
+			ShaderTable.SetLocalShaderParameters(ShaderTableOffset, RecordIndex, BindOffset, DescriptorTable);
+		}
 	}
 
-	FD3D12ConstantBufferView* CreateTransientConstantBuffer(FD3D12ResourceLocation& ResourceLocation, const void* Data, uint32 DataSize)
+	FD3D12ConstantBufferView* SetLooseParameterData(const void* Data, uint32 DataSize, D3D12_GPU_VIRTUAL_ADDRESS& OutGPUVirtualAddress)
 	{
-		// If we see a significant number of transient allocations coming through this path, we should consider
-		// caching constant buffer blocks inside ShaderTable and linearly sub-allocate from them.
-		// If the amount of data is relatively small, it may also be possible to use root constants and avoid extra allocations entirely.
-
-	#if D3D12RHI_USE_CONSTANT_BUFFER_VIEWS
-		FD3D12ConstantBufferView* ConstantBufferView = new FD3D12ConstantBufferView(GetDevice(), nullptr);
-		ShaderTable.WorkerData[WorkerIndex].TransientCBVs.Add(ConstantBufferView);
-	#else // D3D12RHI_USE_CONSTANT_BUFFER_VIEWS
-		FD3D12ConstantBufferView* ConstantBufferView = nullptr;
-	#endif // D3D12RHI_USE_CONSTANT_BUFFER_VIEWS
-
-		FD3D12FastConstantAllocator& Allocator = Device.GetParentAdapter()->GetTransientUniformBufferAllocator();
-		void* MappedData = Allocator.Allocate(DataSize, ResourceLocation, ConstantBufferView);
-
-		FMemory::Memcpy(MappedData, Data, DataSize);
+		bHasLooseParameterData = true;
+		return ShaderTable.SetLooseParameterData(WorkerIndex, RecordIndex, BindingType, Data, DataSize, OutGPUVirtualAddress);
+	}
 
-		ShaderTable.AddResourceReference(ResourceLocation.GetResource(), WorkerIndex);
+	void AddReferencedShaderResource(FD3D12BaseShaderResource* ShaderResource)
+	{
+		ShaderTable.AddReferencedShaderResource(WorkerIndex, RecordIndex, BindingType, ShaderResource);
+	}
 
-		return ConstantBufferView;
+	void AddReferencedTexture(FRHITexture* RHITexture)
+	{
+		ShaderTable.AddReferencedTexture(WorkerIndex, RecordIndex, BindingType, RHITexture);
 	}
 
-	void AddResourceReference(FD3D12Resource* D3D12Resource)
+	void AddReferencedUniformBuffer(uint32 BaseSlotIndex, uint32 DescriptorIndex, FD3D12UniformBuffer* UniformBuffer)
 	{
-		ShaderTable.AddResourceReference(D3D12Resource, WorkerIndex);
+		const uint32 OffsetWithinRootSignature = ComputeOffsetWithinRootSignature(BaseSlotIndex, DescriptorIndex);		
+		ShaderTable.AddReferencedUniformBuffer(WorkerIndex, RecordIndex, ShaderTableOffset, OffsetWithinRootSignature, BindingType, UniformBuffer);
 	}
 
 	void AddResourceTransition(FD3D12ShaderResourceView* SRV)
@@ -4378,9 +4973,14 @@ struct FD3D12RayTracingLocalResourceBinder
 	FD3D12RayTracingShaderBindingTableInternal& ShaderTable;
 	FD3D12ExplicitDescriptorCache& DescriptorCache;
 	const FD3D12RootSignature& RootSignature;
+	ERayTracingLocalShaderBindingType BindingType;
 	uint32 ShaderTableOffset = 0;
 	uint32 RecordIndex = ~0u;
 	uint32 WorkerIndex = 0;
+	bool bHasLooseParameterData = false;
+	
+	TArray<FD3D12View*, TInlineAllocator<MAX_SRVS + MAX_UAVS>> ReferencedViews;
+	TArray<FD3D12Resource*, TInlineAllocator<MAX_CBS>> ReferencedCBResources;
 };
 
 template <typename ResourceBinderType>
@@ -4421,6 +5021,7 @@ static bool SetRayTracingShaderResources(
 		D3D12_CPU_DESCRIPTOR_HANDLE LocalCBVs[MAX_CBS];
 #endif
 		D3D12_GPU_VIRTUAL_ADDRESS RemoteCBVs[MAX_CBS];
+		FD3D12UniformBuffer* UniformBuffers[MAX_CBS] = {};
 
 		D3D12_CPU_DESCRIPTOR_HANDLE LocalSRVs[MAX_SRVS];
 		D3D12_CPU_DESCRIPTOR_HANDLE LocalUAVs[MAX_UAVS];
@@ -4433,7 +5034,6 @@ static bool SetRayTracingShaderResources(
 		uint32 UAVVersions[MAX_SRVS];
 		uint32 SamplerVersions[MAX_SRVS];
 
-		TArray<FD3D12Resource*, TInlineAllocator<MAX_CBS + MAX_SRVS + MAX_UAVS>> ReferencedResources;
 		TArray<FD3D12RayTracingScene*, TInlineAllocator<1>> ReferencedRayTracingScenes;
 
 		uint64 BoundSRVMask = 0;
@@ -4456,7 +5056,7 @@ static bool SetRayTracingShaderResources(
 				BoundUAVMask |= 1ull << Index;
 			}
 
-			ReferencedResources.Add(UAV->GetResource());
+			Binder.AddReferencedShaderResource(UAV->GetShaderResource());
 			Binder.AddResourceTransition(UAV);
 		}
 
@@ -4475,7 +5075,7 @@ static bool SetRayTracingShaderResources(
 				BoundSRVMask |= 1ull << Index;
 			}
 
-			ReferencedResources.Add(SRV->GetResource());
+			Binder.AddReferencedShaderResource(SRV->GetShaderResource());
 			Binder.AddResourceTransition(SRV);
 
 			FD3D12RayTracingScene* ReferencedRayTracingScene = SRV->GetRayTracingScene();
@@ -4504,7 +5104,7 @@ static bool SetRayTracingShaderResources(
 				BoundSRVMask |= 1ull << Index;
 			}
 
-			ReferencedResources.Add(SRV->GetResource());
+			Binder.AddReferencedTexture(RHITexture);
 			Binder.AddResourceTransition(SRV);
 		}
 
@@ -4526,7 +5126,7 @@ static bool SetRayTracingShaderResources(
 
 			BoundSRVMask |= 1ull << Index;
 
-			ReferencedResources.Add(SRV->GetResource());
+			Binder.AddReferencedShaderResource(SRV->GetShaderResource());
 			Binder.AddResourceTransition(SRV);
 		}
 
@@ -4605,7 +5205,7 @@ static bool SetRayTracingShaderResources(
 			Bindings.RemoteCBVs[CBVIndex] = CBV->ResourceLocation.GetGPUVirtualAddress();
 			Bindings.BoundCBVMask |= 1ull << CBVIndex;
 
-			Bindings.ReferencedResources.Add(CBV->ResourceLocation.GetResource());
+			Bindings.UniformBuffers[CBVIndex] = CBV;
 		}
 	}
 
@@ -4649,17 +5249,15 @@ static bool SetRayTracingShaderResources(
 
 	if (InLooseParameterData && Shader->UsesGlobalUniformBuffer())
 	{
-		const uint32 CBVIndex = 0; // Global uniform buffer is always assumed to be in slot 0
-
-		FD3D12ResourceLocation ResourceLocation(Binder.GetDevice());
-		FD3D12ConstantBufferView* ConstantBufferView = Binder.CreateTransientConstantBuffer(ResourceLocation, InLooseParameterData, InLooseParameterDataSize);
+		D3D12_GPU_VIRTUAL_ADDRESS LooseParameterGPUVirtualAddress = 0;
+		FD3D12ConstantBufferView* ConstantBufferView = Binder.SetLooseParameterData(InLooseParameterData, InLooseParameterDataSize, LooseParameterGPUVirtualAddress);
 
 	#if D3D12RHI_USE_CONSTANT_BUFFER_VIEWS
-		Bindings.LocalCBVs[CBVIndex] = ConstantBufferView->GetOfflineCpuHandle();
+		Bindings.LocalCBVs[LooseParameterCBVIndex] = ConstantBufferView->GetOfflineCpuHandle();
 	#endif // D3D12RHI_USE_CONSTANT_BUFFER_VIEWS
-		Bindings.RemoteCBVs[CBVIndex] = ResourceLocation.GetGPUVirtualAddress();
+		Bindings.RemoteCBVs[LooseParameterCBVIndex] = LooseParameterGPUVirtualAddress;
 
-		Bindings.BoundCBVMask |= 1ull << CBVIndex;
+		Bindings.BoundCBVMask |= 1ull << LooseParameterCBVIndex;
 	}
 
 	// Validate that all resources required by the shader are set
@@ -4738,6 +5336,12 @@ static bool SetRayTracingShaderResources(
 				const uint64 SlotMask = (1ull << i);
 				D3D12_GPU_VIRTUAL_ADDRESS BufferAddress = (Bindings.BoundCBVMask & SlotMask) ? Bindings.RemoteCBVs[i] : 0;
 				Binder.SetRootCBV(BindSlot, i, BufferAddress);
+				
+				// Also set the uniform buffer in case it's valid so it can be notified about changes when it's a persistent binding
+				if (Bindings.UniformBuffers[i])
+				{
+					Binder.AddReferencedUniformBuffer(BindSlot, i, Bindings.UniformBuffers[i]);					
+				}
 			}
 		}
 	}
@@ -4760,11 +5364,6 @@ static bool SetRayTracingShaderResources(
 		Binder.SetRootDescriptorTable(BindSlot, ResourceDescriptorTableBaseGPU);
 	}
 
-	for (FD3D12Resource* Resource : Bindings.ReferencedResources)
-	{
-		Binder.AddResourceReference(Resource);
-	}
-
 	for (FD3D12RayTracingScene* RayTracingScene : Bindings.ReferencedRayTracingScenes)
 	{
 		Binder.AddRayTracingSceneReference(RayTracingScene);
@@ -5061,6 +5660,7 @@ static void SetRayTracingHitGroup(
 	uint32 NumUniformBuffers, FRHIUniformBuffer* const* UniformBuffers,
 	uint32 LooseParameterDataSize, const void* LooseParameterData,
 	uint32 UserData,
+	ERayTracingLocalShaderBindingType BindingType,
 	uint32 WorkerIndex)
 {
 	const uint32 GPUIndex = Device->GetGPUIndex();
@@ -5080,13 +5680,24 @@ static void SetRayTracingHitGroup(
 		const uint32 NumGeometrySegments = Geometry->GetNumSegments();
 		checkf(GeometrySegmentIndex < NumGeometrySegments, TEXT("Segment %d is out of range for ray tracing geometry '%s' that contains %d segments"),
 			GeometrySegmentIndex, Geometry->DebugName.IsNone() ? TEXT("UNKNOWN") : *Geometry->DebugName.ToString(), NumGeometrySegments);
+
+		// If persistent then all uniform buffers need to be allocated multiframe
+		if (ShaderTable->Lifetime == ERayTracingShaderBindingTableLifetime::Persistent)
+		{
+			for (uint32 UBIndex = 0; UBIndex < NumUniformBuffers; ++UBIndex)
+			{
+				FRHIUniformBuffer* Resource = UniformBuffers[UBIndex];
+				if (Resource)
+				{
+					FD3D12UniformBuffer* UB = FD3D12CommandContext::RetrieveObject<FD3D12UniformBuffer>(Resource, 0);
+					check(UB->UniformBufferUsage == UniformBuffer_MultiFrame);
+				}
+			}
+		}
 	}
 #endif // DO_CHECK
 
-	FD3D12HitGroupSystemParameters SystemParameters = Geometry->HitGroupSystemParameters[GPUIndex][GeometrySegmentIndex];
-	SystemParameters.RootConstants.UserData = UserData;
-
-	ShaderTable->SetHitGroupSystemParameters(RecordIndex, SystemParameters);
+	ShaderTable->SetHitGroupSystemParameters(WorkerIndex, BindingType, RecordIndex, Geometry, UserData, GeometrySegmentIndex);
 
 	const FD3D12RayTracingShader* Shader = Pipeline->HitGroupShaders.Shaders[HitGroupIndex];
 
@@ -5096,7 +5707,8 @@ static void SetRayTracingHitGroup(
 	const bool bCanUseRecordCache = GRayTracingCacheShaderRecords
 		&& LooseParameterDataSize == 0 // loose parameters end up in unique constant buffers, so SBT records can't be shared
 		&& NumUniformBuffers > 0 // there is no benefit from cache if no resources are being bound
-		&& NumUniformBuffers <= CacheKey.MaxUniformBuffers;
+		&& NumUniformBuffers <= CacheKey.MaxUniformBuffers
+		&& (BindingType == ERayTracingLocalShaderBindingType::Transient || ShaderTable->Lifetime == ERayTracingShaderBindingTableLifetime::Transient); // Only transient SBTs or Transient records can use the record cache because otherwise the per record listeners need to tracked and copied as well
 
 	if (bCanUseRecordCache)
 	{
@@ -5113,7 +5725,7 @@ static void SetRayTracingHitGroup(
 		}
 	}
 
-	FD3D12RayTracingLocalResourceBinder ResourceBinder(*Device, *ShaderTable, *(Shader->LocalRootSignature), RecordIndex, WorkerIndex, ERayTracingBindingType::HitGroup);
+	FD3D12RayTracingLocalResourceBinder ResourceBinder(*Device, *ShaderTable, *(Shader->LocalRootSignature), BindingType, RecordIndex, WorkerIndex, ERayTracingBindingType::HitGroup);
 	const bool bResourcesBound = SetRayTracingShaderResources(Shader, Shader->LocalRootSignature,		
 		0, nullptr,	// BindlessParameters
 		0, nullptr, // Textures
@@ -5129,10 +5741,15 @@ static void SetRayTracingHitGroup(
 		ShaderTable->WorkerData[WorkerIndex].ShaderRecordCache.FindOrAdd(CacheKey, RecordIndex);
 	}
 
-	ShaderTable->SetHitGroupIdentifier(RecordIndex,
-		bResourcesBound
-		? Pipeline->HitGroupShaders.Identifiers[HitGroupIndex]
-		: FD3D12ShaderIdentifier::Null);
+	FD3D12ShaderIdentifier ShaderIdentifier = bResourcesBound ? Pipeline->HitGroupShaders.Identifiers[HitGroupIndex] : FD3D12ShaderIdentifier::Null;
+	if (BindingType != ERayTracingLocalShaderBindingType::Validation)
+	{
+		ShaderTable->SetHitGroupIdentifier(RecordIndex, ShaderIdentifier);
+	}
+	else
+	{
+		ShaderTable->CompareHitGroupIdentifier(RecordIndex, ShaderIdentifier);
+	}
 }
 
 static void SetRayTracingCallableShader(
@@ -5155,7 +5772,7 @@ static void SetRayTracingCallableShader(
 	{
 		const FD3D12RayTracingShader* Shader = Pipeline->CallableShaders.Shaders[ShaderIndexInPipeline];
 
-		FD3D12RayTracingLocalResourceBinder ResourceBinder(*Device, *ShaderTable, *(Shader->LocalRootSignature), RecordIndex, WorkerIndex, ERayTracingBindingType::CallableShader);
+		FD3D12RayTracingLocalResourceBinder ResourceBinder(*Device, *ShaderTable, *(Shader->LocalRootSignature), ERayTracingLocalShaderBindingType::Transient, RecordIndex, WorkerIndex, ERayTracingBindingType::CallableShader);
 		const bool bResourcesBound = SetRayTracingShaderResources(Shader, Shader->LocalRootSignature,
 			0, nullptr,	// BindlessParameters
 			0, nullptr, // Textures
@@ -5191,7 +5808,7 @@ static void SetRayTracingMissShader(
 
 	const FD3D12RayTracingShader* Shader = Pipeline->MissShaders.Shaders[ShaderIndexInPipeline];
 
-	FD3D12RayTracingLocalResourceBinder ResourceBinder(*Device, *ShaderTable, *(Shader->LocalRootSignature), RecordIndex, WorkerIndex, ERayTracingBindingType::MissShader);
+	FD3D12RayTracingLocalResourceBinder ResourceBinder(*Device, *ShaderTable, *(Shader->LocalRootSignature), ERayTracingLocalShaderBindingType::Transient, RecordIndex, WorkerIndex, ERayTracingBindingType::MissShader);
 	const bool bResourcesBound = SetRayTracingShaderResources(Shader, Shader->LocalRootSignature,
 		0, nullptr,	// BindlessParameters
 		0, nullptr, // Textures
@@ -5221,7 +5838,7 @@ void FD3D12CommandContext::RHISetBindingsOnShaderBindingTable(
 
 	// Pipeline shouldn't contain any shaders which have bigger local data size then currently set in the initializer
 	// (Otherwise changing of local binding data size would need to supported)
-	check(Pipeline->MaxLocalRootSignatureSize <= SBT->GetInitializer().LocalBindingDataSize);
+	check(Pipeline == nullptr || Pipeline->MaxLocalRootSignatureSize <= SBT->GetInitializer().LocalBindingDataSize);
 
 	FD3D12RayTracingShaderBindingTableInternal* ShaderTableForDevice = SBT->GetTableForDevice(GetParentDevice());
 
@@ -5250,20 +5867,33 @@ void FD3D12CommandContext::RHISetBindingsOnShaderBindingTable(
 		if (BindingType == ERayTracingBindingType::HitGroup)
 		{
 			const FD3D12RayTracingGeometry* Geometry = FD3D12DynamicRHI::ResourceCast(Binding.Geometry);
+			if (Geometry || ShaderTableForDevice->HitGroupIndexingMode == ERayTracingHitGroupIndexingMode::Disallow)
+			{
+				//UE_LOG(LogD3D12RHI, Log, TEXT("Set hit record data for RecordIndex %d on SBT %#016llx with mode: %d"), Binding.RecordIndex, ShaderTableForDevice, Binding.BindingType);
+
+				SetRayTracingHitGroup(Device,
+				    ShaderTableForDevice, Binding.RecordIndex,
+				    Pipeline, Binding.ShaderIndexInPipeline,
+				    Geometry, Binding.SegmentIndex,
+				    Binding.NumUniformBuffers,
+				    Binding.UniformBuffers,
+				    Binding.LooseParameterDataSize,
+				    Binding.LooseParameterData,
+				    Binding.UserData,
+				    Binding.BindingType,
+				    Context.WorkerIndex);
+			}
+			else
+			{
+				//UE_LOG(LogD3D12RHI, Log, TEXT("Clear hit record data for RecordIndex %d on SBT %#016llx"), Binding.RecordIndex, ShaderTableForDevice);
 
-			SetRayTracingHitGroup(Device, 
-				ShaderTableForDevice, Binding.RecordIndex,
-				Pipeline, Binding.ShaderIndexInPipeline,
-				Geometry, Binding.SegmentIndex,
-				Binding.NumUniformBuffers,
-				Binding.UniformBuffers,
-				Binding.LooseParameterDataSize,
-				Binding.LooseParameterData,
-				Binding.UserData,
-				Context.WorkerIndex);
+				check(Binding.BindingType == ERayTracingLocalShaderBindingType::Persistent);
+				ShaderTableForDevice->ClearHitRecordData(Context.WorkerIndex, Binding.RecordIndex);
+			}
 		}
 		else if (BindingType == ERayTracingBindingType::CallableShader)
 		{
+			check(Binding.BindingType == ERayTracingLocalShaderBindingType::Transient);
 			SetRayTracingCallableShader(Device,
 				ShaderTableForDevice, Binding.RecordIndex,
 				Pipeline, Binding.ShaderIndexInPipeline,
@@ -5276,6 +5906,7 @@ void FD3D12CommandContext::RHISetBindingsOnShaderBindingTable(
 		}
 		else if (BindingType == ERayTracingBindingType::MissShader)
 		{
+			check(Binding.BindingType == ERayTracingLocalShaderBindingType::Transient);
 			SetRayTracingMissShader(Device, 
 				ShaderTableForDevice, Binding.RecordIndex,
 				Pipeline, Binding.ShaderIndexInPipeline,
diff --git a/Engine/Source/Runtime/D3D12RHI/Private/D3D12RayTracing.h b/Engine/Source/Runtime/D3D12RHI/Private/D3D12RayTracing.h
--- a/Engine/Source/Runtime/D3D12RHI/Private/D3D12RayTracing.h
+++ b/Engine/Source/Runtime/D3D12RHI/Private/D3D12RayTracing.h
@@ -13,6 +13,15 @@ static_assert(sizeof(FD3D12_GPU_VIRTUAL_ADDRESS) == sizeof(D3D12_GPU_VIRTUAL_ADD
 class FD3D12RayTracingPipelineState;
 class FD3D12RayTracingShaderBindingTable;
 
+class FD3D12RayTracingGeometry;
+
+/** Persistent SBT needs to be notified about hit group parameter changes because those are cached in the SBT - if persistent bindless handles are used then this could be removed */
+struct ID3D12RayTracingGeometryUpdateListener
+{
+	virtual void RemoveListener(FD3D12RayTracingGeometry* InGeometry) = 0;
+	virtual void HitGroupParametersUpdated(FD3D12RayTracingGeometry* InGeometry) = 0;
+};
+
 class FD3D12RayTracingGeometry : public FRHIRayTracingGeometry, public FD3D12AdapterChild, public FD3D12ShaderResourceRenameListener, public FNoncopyable
 {
 public:
@@ -88,6 +97,42 @@ public:
 	TArray<D3D12_RAYTRACING_GEOMETRY_DESC, TInlineAllocator<1>> GeometryDescs;
 
 	uint64 AccelerationStructureCompactedSize = 0;
+		
+public:
+
+	void AddUpdateListener(ID3D12RayTracingGeometryUpdateListener* InUpdateListener) const
+	{
+		FScopeLock Lock(&UpdateListenersCS);
+		check(!UpdateListeners.Contains(InUpdateListener));
+		UpdateListeners.Add(InUpdateListener);
+	}
+
+	void RemoveUpdateListener(ID3D12RayTracingGeometryUpdateListener* InUpdateListener) const
+	{
+		FScopeLock Lock(&UpdateListenersCS);
+		uint32 Removed = UpdateListeners.Remove(InUpdateListener);
+
+		checkf(Removed == 1, TEXT("Should have exactly one registered listener during remove (same listener shouldn't registered twice and we shouldn't call this if not registered"));
+	}
+
+	bool HasListeners() const
+	{
+		FScopeLock Lock(&UpdateListenersCS);
+		return UpdateListeners.Num() != 0;
+	}
+
+	void HitGroupParamatersUpdated()
+	{
+		FScopeLock Lock(&UpdateListenersCS);
+		for (ID3D12RayTracingGeometryUpdateListener* UpdateListener : UpdateListeners)
+		{
+			UpdateListener->HitGroupParametersUpdated(this);
+		}
+	}
+
+private:
+	mutable FCriticalSection UpdateListenersCS;
+	mutable TArray<ID3D12RayTracingGeometryUpdateListener*> UpdateListeners;
 };
 
 class FD3D12RayTracingScene : public FRHIRayTracingScene, public FD3D12AdapterChild, public FNoncopyable
diff --git a/Engine/Source/Runtime/D3D12RHI/Private/D3D12Resources.h b/Engine/Source/Runtime/D3D12RHI/Private/D3D12Resources.h
--- a/Engine/Source/Runtime/D3D12RHI/Private/D3D12Resources.h
+++ b/Engine/Source/Runtime/D3D12RHI/Private/D3D12Resources.h
@@ -926,9 +926,21 @@ public:
 	}
 };
 
+class FD3D12UniformBuffer;
+
+/** Resource which might needs to be notified about changes on dependent resources (Views, RTGeometryObject, Cached binding tables) */
+struct ID3D12UniformBufferUpdateListener
+{
+	virtual void RemoveListener(FD3D12UniformBuffer* InUniformBufferToRemove) = 0;
+	virtual void UniformBufferUpdated(FRHICommandListBase& CmdList, FD3D12UniformBuffer* InUpdatedUniformBuffer) = 0;
+};
+
 /** Uniform buffer resource class. */
 class FD3D12UniformBuffer : public FRHIUniformBuffer, public FD3D12DeviceChild, public FD3D12LinkedAdapterObject<FD3D12UniformBuffer>
 {
+private:
+	mutable FCriticalSection UpdateListenersCS;
+	TArray<ID3D12UniformBufferUpdateListener*> UpdateListeners;
 public:
 #if D3D12RHI_USE_CONSTANT_BUFFER_VIEWS
 	class FD3D12ConstantBufferView* View = nullptr;
@@ -953,6 +965,36 @@ public:
 	// Provides public non-const access to ResourceTable.
 	// @todo refactor uniform buffers to perform updates as a member function, so this isn't necessary.
 	TArray<TRefCountPtr<FRHIResource>>& GetResourceTable() { return ResourceTable; }
+		
+	void AddUpdateListener(ID3D12UniformBufferUpdateListener* InUpdateListener)
+	{
+		FScopeLock Lock(&UpdateListenersCS);
+		check(!UpdateListeners.Contains(InUpdateListener));
+		UpdateListeners.Add(InUpdateListener);
+	}
+
+	void RemoveUpdateListener(ID3D12UniformBufferUpdateListener* InUpdateListener)
+	{
+		FScopeLock Lock(&UpdateListenersCS);
+		uint32 Removed = UpdateListeners.Remove(InUpdateListener);
+
+		checkf(Removed == 1, TEXT("Should have exactly one registered listener during remove (same listener shouldn't registered twice and we shouldn't call this if not registered"));
+	}
+
+	bool HasListeners() const
+	{
+		FScopeLock Lock(&UpdateListenersCS);
+		return UpdateListeners.Num() != 0;
+	}
+
+	void UniformBufferUpdated(FRHICommandListBase& CmdList)
+	{
+		FScopeLock Lock(&UpdateListenersCS);
+		for (ID3D12UniformBufferUpdateListener* UpdateListener : UpdateListeners)
+		{
+			UpdateListener->UniformBufferUpdated(CmdList, this);
+		}
+	}
 };
 
 class FD3D12Buffer : public FRHIBuffer, public FD3D12BaseShaderResource, public FD3D12LinkedAdapterObject<FD3D12Buffer>
diff --git a/Engine/Source/Runtime/D3D12RHI/Private/D3D12UniformBuffer.cpp b/Engine/Source/Runtime/D3D12RHI/Private/D3D12UniformBuffer.cpp
--- a/Engine/Source/Runtime/D3D12RHI/Private/D3D12UniformBuffer.cpp
+++ b/Engine/Source/Runtime/D3D12RHI/Private/D3D12UniformBuffer.cpp
@@ -151,6 +151,9 @@ FRHICOMMAND_MACRO(FRHICommandD3D12UpdateUniformBuffer)
 		const uint32 NumBytes = Align(UniformBuffer->GetLayout().ConstantBufferSize, D3D12_CONSTANT_BUFFER_DATA_PLACEMENT_ALIGNMENT);
 		UniformBuffer->View->CreateView(&UniformBuffer->ResourceLocation, 0, NumBytes);
 #endif
+
+		// Notify the listeners now that the resource location on the uniform buffer has been updated
+		UniformBuffer->UniformBufferUpdated(CmdList);
 	}
 };
 
@@ -232,6 +235,19 @@ FD3D12UniformBuffer::~FD3D12UniformBuffer()
 {
 	check(IsInRHIThread() || IsInRenderingThread());
 
+	if (!UpdateListeners.IsEmpty())
+	{
+		//UE_LOG(LogD3D12RHI, Log, TEXT("Deleting uniform buffer %#016llx with GPU address: \"0x%llX\" and %d listeners still registered"), this, ResourceLocation.GetGPUVirtualAddress(), UpdateListeners.Num());
+
+		// Request remove of listener to this uniform buffer - uniform buffers can be deleted before the cached MDCs referencing
+		// the uniform buffers are deleted because requests are still pending for scene proxy removal
+		for (ID3D12UniformBufferUpdateListener* UpdateListener : UpdateListeners)
+		{
+			UpdateListener->RemoveListener(this);
+		}
+		UpdateListeners.Empty();
+	}
+
 	int64 BufferSize = ResourceLocation.GetSize();
 	DEC_MEMORY_STAT_BY(STAT_UniformBufferMemory, BufferSize);
 
diff --git a/Engine/Source/Runtime/D3D12RHI/Private/D3D12View.h b/Engine/Source/Runtime/D3D12RHI/Private/D3D12View.h
--- a/Engine/Source/Runtime/D3D12RHI/Private/D3D12View.h
+++ b/Engine/Source/Runtime/D3D12RHI/Private/D3D12View.h
@@ -215,6 +215,7 @@ public:
 	};
 
 	FD3D12Resource*                         GetResource        () const { check(IsInitialized()); return ResourceInfo.Resource;         }
+	FD3D12BaseShaderResource*               GetShaderResource  () const { check(IsInitialized()); return ResourceInfo.BaseResource;     }
 	FD3D12ResourceLocation*                 GetResourceLocation() const { check(IsInitialized()); return ResourceInfo.ResourceLocation; }
 	TConstArrayView<FD3D12ResidencyHandle*> GetResidencyHandles() const { check(IsInitialized()); return ResourceInfo.Resource ? ResourceInfo.Resource->GetResidencyHandles() : TConstArrayView<FD3D12ResidencyHandle*>(); }
 	FD3D12ViewSubset const&                 GetViewSubset      () const { check(IsInitialized()); return ViewSubset;                    }
diff --git a/Engine/Source/Runtime/RHI/Public/RHICommandList.h b/Engine/Source/Runtime/RHI/Public/RHICommandList.h
--- a/Engine/Source/Runtime/RHI/Public/RHICommandList.h
+++ b/Engine/Source/Runtime/RHI/Public/RHICommandList.h
@@ -255,8 +255,16 @@ struct FRayTracingShaderBindings
 	TArray<FRHIShaderParameterResource> BindlessParameters;
 };
 
+enum class ERayTracingLocalShaderBindingType : uint8
+{
+	Persistent,		//< Binding contains persistent data
+	Transient,		//< Binding contains transient data
+	Validation		//< Binding only used for validating persistently stored data in the SBT
+};
+
 struct FRayTracingLocalShaderBindings
 {
+	ERayTracingLocalShaderBindingType BindingType = ERayTracingLocalShaderBindingType::Transient;
 	const FRHIRayTracingGeometry* Geometry = nullptr;
 	uint32 SegmentIndex = 0;
 	uint32 RecordIndex = 0;
diff --git a/Engine/Source/Runtime/RHI/Public/RHIResources.h b/Engine/Source/Runtime/RHI/Public/RHIResources.h
--- a/Engine/Source/Runtime/RHI/Public/RHIResources.h
+++ b/Engine/Source/Runtime/RHI/Public/RHIResources.h
@@ -3349,6 +3349,13 @@ enum class ERayTracingAccelerationStructureFlags
 };
 ENUM_CLASS_FLAGS(ERayTracingAccelerationStructureFlags);
 
+enum class ERayTracingShaderBindingTableLifetime : uint8
+{
+	Transient,						//< SBT will be reallocated each frame
+	Persistent,						//< SBT will be persistently stored and only new or changed bindings will be set
+};
+ENUM_CLASS_FLAGS(ERayTracingShaderBindingTableLifetime);
+
 enum class ERayTracingShaderBindingMode : uint8
 {
 	Disabled	= 0,				//< No binding data at all
@@ -3366,6 +3373,9 @@ ENUM_CLASS_FLAGS(ERayTracingHitGroupIndexingMode);
 
 struct FRayTracingShaderBindingTableInitializer
 {
+	// Defines lifetime of the shader binding table
+	ERayTracingShaderBindingTableLifetime Lifetime = ERayTracingShaderBindingTableLifetime::Transient;
+
 	// Defines which types of binding data needs to be stored in the SBT (Inline and/or RTPSO)
 	ERayTracingShaderBindingMode ShaderBindingMode = ERayTracingShaderBindingMode::Disabled;
 	
