Add virtual destructor to listeners

diff --git a/Engine/Source/Runtime/D3D12RHI/Private/D3D12RayTracing.cpp b/Engine/Source/Runtime/D3D12RHI/Private/D3D12RayTracing.cpp
--- a/Engine/Source/Runtime/D3D12RHI/Private/D3D12RayTracing.cpp
+++ b/Engine/Source/Runtime/D3D12RHI/Private/D3D12RayTracing.cpp
@@ -1771,6 +1771,8 @@ public:
 		{
 		}
 
+		virtual ~FRecordUpdateUniformBufferListener() = default;
+
 		void Init(FD3D12UniformBuffer* InUniformBuffer, uint32 InShaderTableOffset, uint32 InRecordIndex, uint32 InOffsetWithinRootSignature)
 		{
 			UniformBuffer = InUniformBuffer;
@@ -1811,6 +1813,8 @@ public:
 		uint32 UserData = 0;
 		uint32 SegmentIndex = 0;
 
+		virtual ~FRecordUpdateRayTracingGeometryListener() = default;
+
 		void Init(FD3D12RayTracingShaderBindingTableInternal* InShaderTable, const FD3D12RayTracingGeometry* InGeometry, uint32 InRecordIndex, uint32 InUserData, uint32 InSegmentIndex)
 		{
 			check(Geometry == nullptr);
