D3D12 RT Validation: Move breadcrumbs out of recursive cmd list.

#rnx

diff --git a/Engine/Source/Runtime/D3D12RHI/Private/D3D12RayTracing.cpp b/Engine/Source/Runtime/D3D12RHI/Private/D3D12RayTracing.cpp
--- a/Engine/Source/Runtime/D3D12RHI/Private/D3D12RayTracing.cpp
+++ b/Engine/Source/Runtime/D3D12RHI/Private/D3D12RayTracing.cpp
@@ -4350,6 +4350,8 @@ void BuildAccelerationStructure(FD3D12CommandContext& CommandContext,
 
 	if (ShouldRunRayTracingGPUValidation())
 	{
+		RHI_BREADCRUMB_EVENT(CommandContext, "RTSceneValidation");
+
 		TRHICommandList_RecursiveHazardous<FD3D12CommandContext> RHICmdList(&CommandContext);
 		uint32 InstanceBufferStride = GRHIRayTracingInstanceDescriptorSize;
 		PRAGMA_DISABLE_DEPRECATION_WARNINGS
@@ -4631,8 +4633,9 @@ void FD3D12CommandContext::RHIBuildAccelerationStructures(const TArrayView<const
 
 	if (ShouldRunRayTracingGPUValidation())
 	{
-		TRHICommandList_RecursiveHazardous<FD3D12CommandContext> RHICmdList(this);
+		RHI_BREADCRUMB_EVENT(*this, "RTGeometryValidation");
 
+		TRHICommandList_RecursiveHazardous<FD3D12CommandContext> RHICmdList(this);
 		for (const FRayTracingGeometryBuildParams& P : Params)
 		{
 			FRayTracingValidateGeometryBuildParamsCS::Dispatch(RHICmdList, P);
diff --git a/Engine/Source/Runtime/RenderCore/Private/RayTracingValidationShaders.cpp b/Engine/Source/Runtime/RenderCore/Private/RayTracingValidationShaders.cpp
--- a/Engine/Source/Runtime/RenderCore/Private/RayTracingValidationShaders.cpp
+++ b/Engine/Source/Runtime/RenderCore/Private/RayTracingValidationShaders.cpp
@@ -30,8 +30,6 @@ void FRayTracingValidateGeometryBuildParamsCS::Dispatch(FRHICommandList& RHICmdL
 		return;
 	}
 
-	SCOPED_DRAW_EVENTF(RHICmdList, RTGeometryValidation, TEXT("RTGeometryValidation - %s"), Initializer.DebugName);
-
 	const uint32 IndexStride = Initializer.IndexBuffer->GetStride();
 
 	const FRawBufferShaderResourceViewInitializer IBViewInitializer(Initializer.IndexBuffer);
@@ -100,8 +98,6 @@ void FRayTracingValidateSceneBuildParamsCS::Dispatch(FRHICommandList& RHICmdList
 	FRHIComputeShader* ShaderRHI = ComputeShader.GetComputeShader();
 	SetComputePipelineState(RHICmdList, ShaderRHI);
 
-	SCOPED_DRAW_EVENT(RHICmdList, RTSceneValidation);
-
 	{
 		FRHIBatchedShaderParameters& BatchedParameters = RHICmdList.GetScratchShaderParameters();
 
