Fix LLM tag for reserved resources: LLM_REALLOC_SCOPE allows to retrieve the tag under which an original allocation happened. Reserved resources themselves do not actually allocate memory but FD3D12Adapter::CreateReservedResource creates a new FD3D12Resource / FD3D12ReservedResourceData that are tracked by LLM.
Thanks to these pointers, it's possible find the original tag back and use it when commiting parts of the reserved resource.
This fixes some Nanite allocations ending up in RHIMisc

diff --git a/Engine/Source/Runtime/D3D12RHI/Private/D3D12Resources.cpp b/Engine/Source/Runtime/D3D12RHI/Private/D3D12Resources.cpp
--- a/Engine/Source/Runtime/D3D12RHI/Private/D3D12Resources.cpp
+++ b/Engine/Source/Runtime/D3D12RHI/Private/D3D12Resources.cpp
@@ -196,6 +196,7 @@ void FD3D12Resource::CommitReservedResource(ID3D12CommandQueue* D3DCommandQueue,
 
 	check(Desc.bReservedResource);
 	check(ReservedResourceData.IsValid());
+	LLM_REALLOC_SCOPE(ReservedResourceData.Get());
 
 	checkf(GRHIGlobals.ReservedResources.Supported,
 		TEXT("Current RHI does not support reserved resources"));
