Enable async dispatch rays on windows d3d12

diff --git a/Engine/Source/Runtime/D3D12RHI/Private/Windows/WindowsD3D12Device.cpp b/Engine/Source/Runtime/D3D12RHI/Private/Windows/WindowsD3D12Device.cpp
--- a/Engine/Source/Runtime/D3D12RHI/Private/Windows/WindowsD3D12Device.cpp
+++ b/Engine/Source/Runtime/D3D12RHI/Private/Windows/WindowsD3D12Device.cpp
@@ -1725,6 +1725,12 @@ void FD3D12DynamicRHI::Init()
 	GRHISupportsShaderBundleWorkGraphDispatch = GRHISupportsShaderWorkGraphsTier1;
 	GRHISupportsShaderBundleParallel = false; // TODO: FD3D12ExplicitDescriptorHeap::UpdateSyncPoint() is not safe for parallel translate due to frame fencing
 
+	// PC D3D12 support async dispatch rays calls
+	if (GRHIGlobals.RayTracing.Supported)
+	{
+		GRHIGlobals.RayTracing.SupportsAsyncRayTraceDispatch = true;
+	}
+
 	// Indicate that the RHI needs to use the engine's deferred deletion queue.
 	GRHINeedsExtraDeletionLatency = true;
 
