PR #12437: Support cubemap textures in D3D12's RHIAliasTextureResources()

diff --git a/Engine/Source/Runtime/D3D12RHI/Private/D3D12Texture.cpp b/Engine/Source/Runtime/D3D12RHI/Private/D3D12Texture.cpp
--- a/Engine/Source/Runtime/D3D12RHI/Private/D3D12Texture.cpp
+++ b/Engine/Source/Runtime/D3D12RHI/Private/D3D12Texture.cpp
@@ -1502,6 +1502,8 @@ void FD3D12Texture::AliasResources(FD3D12Texture* Texture)
 		DepthStencilViews[Index] = Texture->DepthStencilViews[Index];
 	}
 
+	bCreatedRTVsPerSlice = Texture->bCreatedRTVsPerSlice;
+	RTVArraySizePerMip = Texture->RTVArraySizePerMip;
 	RenderTargetViews.SetNum(Texture->RenderTargetViews.Num());
 	for (int32 Index = 0; Index < Texture->RenderTargetViews.Num(); Index++)
 	{
@@ -2387,9 +2389,15 @@ FD3D12Texture* FD3D12DynamicRHI::CreateTextureFromResource(bool bTextureArray, b
 	TextureResource->AddRef();
 	TextureResource->SetName(TEXT("TextureFromResource"));
 
+	const ETextureDimension Dimension = bTextureArray
+		? (bCubeTexture ? ETextureDimension::TextureCubeArray : ETextureDimension::Texture2DArray)
+		: (bCubeTexture ? ETextureDimension::TextureCube      : ETextureDimension::Texture2D     );
+
 	FRHITextureCreateDesc CreateDesc =
-		((SizeZ > 1) ? FRHITextureCreateDesc::Create2DArray(TEXT("TextureFromResource"), FIntPoint(SizeX, SizeY), SizeZ, Format) :
-			FRHITextureCreateDesc::Create2D(TEXT("TextureFromResource"), FIntPoint(SizeX, SizeY), Format))
+		FRHITextureCreateDesc::Create(TEXT("TextureFromResource"), Dimension)
+		.SetExtent(SizeX, SizeY)
+		.SetArraySize(SizeZ)
+		.SetFormat(Format)
 		.SetClearValue(ClearValueBinding)
 		.SetFlags(TexCreateFlags)
 		.SetNumMips(NumMips)
