Promoted d3d12 zero-sized buffer validation to fatal logging

diff --git a/Engine/Source/Runtime/D3D12RHI/Private/D3D12Adapter.h b/Engine/Source/Runtime/D3D12RHI/Private/D3D12Adapter.h
--- a/Engine/Source/Runtime/D3D12RHI/Private/D3D12Adapter.h
+++ b/Engine/Source/Runtime/D3D12RHI/Private/D3D12Adapter.h
@@ -467,7 +467,8 @@ protected:
 		FD3D12Buffer* Buffer,
 		FD3D12ResourceLocation& ResourceLocation,
 		ID3D12ResourceAllocator* ResourceAllocator,
-		const TCHAR* InDebugName);
+		const TCHAR* InDebugName,
+		const FName& OwnerName);
 
 	// Creates default root and execute indirect signatures
 	virtual void CreateCommandSignatures();
diff --git a/Engine/Source/Runtime/D3D12RHI/Private/D3D12Buffer.cpp b/Engine/Source/Runtime/D3D12RHI/Private/D3D12Buffer.cpp
--- a/Engine/Source/Runtime/D3D12RHI/Private/D3D12Buffer.cpp
+++ b/Engine/Source/Runtime/D3D12RHI/Private/D3D12Buffer.cpp
@@ -190,13 +190,21 @@ void FD3D12Adapter::AllocateBuffer(FD3D12Device* Device,
 	FD3D12Buffer* Buffer,
 	FD3D12ResourceLocation& ResourceLocation,
 	ID3D12ResourceAllocator* ResourceAllocator,
-	const TCHAR* InDebugName)
+	const TCHAR* InDebugName,
+	const FName& OwnerName)
 {
 	TRACE_CPUPROFILER_EVENT_SCOPE(D3D12RHI::AllocateBuffer);
 
 	// Explicitly check that the size is nonzero before allowing CreateBuffer to opaquely fail.
-	checkf(Size > 0, TEXT("Attempt to create buffer '%s' with size 0."), InDebugName ? InDebugName : TEXT("(null)"));
-
+	if (Size == 0)
+	{
+		UE_LOG(
+			LogD3D12RHI, Fatal,
+			TEXT("Attempt to create zero-sized buffer '%s', owner '%s', usage 0x%x"),
+			InDebugName ? InDebugName : TEXT("(null)"), *OwnerName.ToString(), static_cast<uint32>(InUsage)
+		);
+	}
+	
 	if (EnumHasAnyFlags(InUsage, BUF_AnyDynamic))
 	{
 		check(ResourceAllocator == nullptr);
@@ -262,7 +270,7 @@ FD3D12Buffer* FD3D12Adapter::CreateRHIBuffer(
 
 			if ((Device->GetGPUIndex() == FirstGPUIndex) || EnumHasAnyFlags(BufferDesc.Usage, BUF_MultiGPUAllocate))
 			{
-				AllocateBuffer(Device, InDesc, BufferDesc.Size, BufferDesc.Usage, InResourceStateMode, InCreateState, Alignment, NewBuffer, NewBuffer->ResourceLocation, ResourceAllocator, InDebugName);
+				AllocateBuffer(Device, InDesc, BufferDesc.Size, BufferDesc.Usage, InResourceStateMode, InCreateState, Alignment, NewBuffer, NewBuffer->ResourceLocation, ResourceAllocator, InDebugName, OwnerName);
 				NewBuffer0 = NewBuffer;
 			}
 			else
@@ -289,7 +297,7 @@ FD3D12Buffer* FD3D12Adapter::CreateRHIBuffer(
 			}
 #endif // NAME_OBJECTS
 
-			AllocateBuffer(Device, InDesc, BufferDesc.Size, BufferDesc.Usage, InResourceStateMode, InCreateState, Alignment, NewBuffer, NewBuffer->ResourceLocation, ResourceAllocator, InDebugName);
+			AllocateBuffer(Device, InDesc, BufferDesc.Size, BufferDesc.Usage, InResourceStateMode, InCreateState, Alignment, NewBuffer, NewBuffer->ResourceLocation, ResourceAllocator, InDebugName, OwnerName);
 			
 			// Unlock immediately if there is no initial data
 			if (!bKeepUnlocked)
