Fixed D3D12 texture locking overestimating number of rows on compressed formats

diff --git a/Engine/Source/Runtime/D3D12RHI/Private/D3D12Texture.cpp b/Engine/Source/Runtime/D3D12RHI/Private/D3D12Texture.cpp
--- a/Engine/Source/Runtime/D3D12RHI/Private/D3D12Texture.cpp
+++ b/Engine/Source/Runtime/D3D12RHI/Private/D3D12Texture.cpp
@@ -1548,8 +1548,10 @@ FRHILockTextureResult FD3D12Texture::Lock(FRHICommandListBase& RHICmdList, const
 
 	const D3D12_RESOURCE_DESC& ResourceDesc = GetResource()->GetDesc();
 
+	uint32 NumRows = 0;
+
 	D3D12_PLACED_SUBRESOURCE_FOOTPRINT PlacedFootprint;
-	Device->GetDevice()->GetCopyableFootprints(&ResourceDesc, Subresource, 1, 0, &PlacedFootprint, nullptr, nullptr, nullptr);
+	Device->GetDevice()->GetCopyableFootprints(&ResourceDesc, Subresource, 1, 0, &PlacedFootprint, &NumRows, nullptr, nullptr);
 
 	// GetCopyableFootprints returns the offset from the start of the resource to the specified subresource, but our staging buffer represents
 	// only the selected subresource, so we need to reset the offset to 0.
@@ -1561,7 +1563,7 @@ FRHILockTextureResult FD3D12Texture::Lock(FRHICommandListBase& RHICmdList, const
 	FRHILockTextureResult Result;
 	Result.Stride = LockedResource->Footprint.RowPitch;
 
-	const uint64 SubresourceSize = LockedResource->Footprint.RowPitch * LockedResource->Footprint.Height * LockedResource->Footprint.Depth;
+	const uint64 SubresourceSize = LockedResource->Footprint.RowPitch * NumRows * LockedResource->Footprint.Depth;
 	Result.ByteCount = SubresourceSize;
 
 	if (FD3D12DynamicRHI::GetD3DRHI()->HandleSpecialLock(Result, this, Arguments))
