Persistent SBT fixes

diff --git a/Engine/Source/Runtime/D3D12RHI/Private/D3D12RayTracing.cpp b/Engine/Source/Runtime/D3D12RHI/Private/D3D12RayTracing.cpp
--- a/Engine/Source/Runtime/D3D12RHI/Private/D3D12RayTracing.cpp
+++ b/Engine/Source/Runtime/D3D12RHI/Private/D3D12RayTracing.cpp
@@ -1402,7 +1402,7 @@ public:
 	~FD3D12RayTracingShaderBindingTableInternal()
 	{
 		// Might be slow - iterate faster somehow?
-		for (uint32 RecordIndex = 0; RecordIndex < NumHitRecords; ++RecordIndex)
+		for (int32 RecordIndex = 0; RecordIndex <= MaxUsedHitRecordIndex; ++RecordIndex)
 		{
 			ClearHitRecordData(0, RecordIndex);
 		}
@@ -1520,6 +1520,8 @@ public:
 
 		const uint32 WriteOffset = HitGroupShaderTableOffset + RecordIndex * LocalRecordStride;
 		WriteData(WriteOffset, ShaderIdentifier.Data, ShaderIdentifierSize);
+
+		MaxUsedHitRecordIndex = FMath::Max(MaxUsedHitRecordIndex, (int32)RecordIndex);
 	}	
 
 	void CompareHitGroupIdentifier(uint32 RecordIndex, const FD3D12ShaderIdentifier& ShaderIdentifier)
@@ -1645,8 +1647,11 @@ public:
 			WorkerData[WorkerIndex].TransitionSRVs.Empty();
 			WorkerData[WorkerIndex].TransitionUAVs.Empty();
 			WorkerData[WorkerIndex].TransitionViewSet.Empty();
+		}
 
-			// Clear the record cache for reuse
+		// Clear the record cache for reuse for all worker thread
+		for (uint32 WorkerIndex = 0; WorkerIndex < MaxBindingWorkers; ++WorkerIndex)
+		{
 			WorkerData[WorkerIndex].ShaderRecordCache.Empty();
 		}
 
@@ -1730,6 +1735,8 @@ public:
 	uint32 NumCallableRecords = 0;
 	uint32 NumMissRecords = 0;
 
+	int32 MaxUsedHitRecordIndex = -1;
+
 	uint32 MissShaderTableOffset = 0;
 	uint32 HitGroupShaderTableOffset = 0;
 	uint32 CallableShaderTableOffset = 0;
@@ -2928,7 +2935,7 @@ public:
 		delete ShaderTablesPerGPU[GPUIndex];
 		ShaderTablesPerGPU[GPUIndex] = nullptr;
 	}
-	
+
 	uint32 GetHitRecordBaseIndex(uint32 GlobalSegmentIndex) const { return GlobalSegmentIndex * Initializer.NumShaderSlotsPerGeometrySegment; }
 
 private:
@@ -4945,7 +4952,7 @@ struct FD3D12RayTracingLocalResourceBinder
 
 	void AddResourceTransition(FD3D12ShaderResourceView* SRV)
 	{
-		if (SRV->GetResource()->RequiresResourceStateTracking())
+		if (SRV->GetResource()->RequiresResourceStateTracking() && ShaderTable.Lifetime == ERayTracingShaderBindingTableLifetime::Transient)
 		{
 			ShaderTable.AddResourceTransition(SRV, WorkerIndex);
 		}
@@ -4953,7 +4960,7 @@ struct FD3D12RayTracingLocalResourceBinder
 
 	void AddResourceTransition(FD3D12UnorderedAccessView* UAV)
 	{
-		if (UAV->GetResource()->RequiresResourceStateTracking())
+		if (UAV->GetResource()->RequiresResourceStateTracking() && ShaderTable.Lifetime == ERayTracingShaderBindingTableLifetime::Transient)
 		{
 			ShaderTable.AddResourceTransition(UAV, WorkerIndex);
 		}
@@ -5733,7 +5740,13 @@ static void SetRayTracingHitGroup(
 			const uint32 OffsetFromRootSignatureStart = sizeof(FD3D12HitGroupSystemParameters);
 			ShaderTable->SetHitGroupIdentifier(RecordIndex, Pipeline->HitGroupShaders.Identifiers[HitGroupIndex]);
 			ShaderTable->CopyHitGroupParameters(RecordIndex, *ExistingRecordIndex, OffsetFromRootSignatureStart);
+
+#if DO_CHECK
+			// Validate all other data against copied data
+			BindingType = ERayTracingLocalShaderBindingType::Validation;
+#else
 			return;
+#endif // DO_CHECK
 		}
 	}
 
@@ -5878,11 +5891,11 @@ void FD3D12CommandContext::RHISetBindingsOnShaderBindingTable(
 
 		if (BindingType == ERayTracingBindingType::HitGroup)
 		{
-			const FD3D12RayTracingGeometry* Geometry = FD3D12DynamicRHI::ResourceCast(Binding.Geometry);
-			if (Geometry || ShaderTableForDevice->HitGroupIndexingMode == ERayTracingHitGroupIndexingMode::Disallow)
+			if (Binding.BindingType != ERayTracingLocalShaderBindingType::Clear)
 			{
 				//UE_LOG(LogD3D12RHI, Log, TEXT("Set hit record data for RecordIndex %d on SBT %#016llx with mode: %d"), Binding.RecordIndex, ShaderTableForDevice, Binding.BindingType);
-
+								
+				const FD3D12RayTracingGeometry* Geometry = FD3D12DynamicRHI::ResourceCast(Binding.Geometry);
 				SetRayTracingHitGroup(Device,
 				    ShaderTableForDevice, Binding.RecordIndex,
 				    Pipeline, Binding.ShaderIndexInPipeline,
@@ -5899,7 +5912,6 @@ void FD3D12CommandContext::RHISetBindingsOnShaderBindingTable(
 			{
 				//UE_LOG(LogD3D12RHI, Log, TEXT("Clear hit record data for RecordIndex %d on SBT %#016llx"), Binding.RecordIndex, ShaderTableForDevice);
 
-				check(Binding.BindingType == ERayTracingLocalShaderBindingType::Persistent);
 				ShaderTableForDevice->ClearHitRecordData(Context.WorkerIndex, Binding.RecordIndex);
 			}
 		}
diff --git a/Engine/Source/Runtime/RHI/Public/RHICommandList.h b/Engine/Source/Runtime/RHI/Public/RHICommandList.h
--- a/Engine/Source/Runtime/RHI/Public/RHICommandList.h
+++ b/Engine/Source/Runtime/RHI/Public/RHICommandList.h
@@ -259,6 +259,7 @@ enum class ERayTracingLocalShaderBindingType : uint8
 {
 	Persistent,		//< Binding contains persistent data
 	Transient,		//< Binding contains transient data
+	Clear,			//< Clear SBT record data		
 	Validation		//< Binding only used for validating persistently stored data in the SBT
 };
 
