Fix use-after-free in FD3D12Device::GetResourceAllocationInfo.

The read lock was released to early, leading to a race condition if another thread inserts an entry into ResourceAllocationInfoMap before the cached data is copied to the return value. Insertions can reallocate, so CachedInfo can end up pointing to freed memory.

#rnx

diff --git a/Engine/Source/Runtime/D3D12RHI/Private/D3D12Device.cpp b/Engine/Source/Runtime/D3D12RHI/Private/D3D12Device.cpp
--- a/Engine/Source/Runtime/D3D12RHI/Private/D3D12Device.cpp
+++ b/Engine/Source/Runtime/D3D12RHI/Private/D3D12Device.cpp
@@ -606,31 +606,36 @@ D3D12_RESOURCE_ALLOCATION_INFO FD3D12Device::GetResourceAllocationInfoUncached(c
 
 D3D12_RESOURCE_ALLOCATION_INFO FD3D12Device::GetResourceAllocationInfo(const FD3D12ResourceDesc& InDesc)
 {
+	D3D12_RESOURCE_ALLOCATION_INFO Result;
+
 	const uint64 Hash = CityHash64((const char*)&InDesc, sizeof(FD3D12ResourceDesc));
 
 	// By default there'll be more threads trying to read this than to write it.
 	ResourceAllocationInfoMapMutex.ReadLock();
 	D3D12_RESOURCE_ALLOCATION_INFO* CachedInfo = ResourceAllocationInfoMap.Find(Hash);
-	ResourceAllocationInfoMapMutex.ReadUnlock();
-
 	if (CachedInfo)
 	{
-		return *CachedInfo;
+		// We need to hold on to the lock until we copy the cached data, because CachedInfo points to the map storage, which
+		// can be reallocated on insertion.
+		Result = *CachedInfo;
+		ResourceAllocationInfoMapMutex.ReadUnlock();
+		return Result;
 	}
-	else
-	{
-		const D3D12_RESOURCE_ALLOCATION_INFO Result = GetResourceAllocationInfoUncached(InDesc);
 
-		ResourceAllocationInfoMapMutex.WriteLock();
-		// Try search again with write lock because could have been added already
-		CachedInfo = ResourceAllocationInfoMap.Find(Hash);
-		if (CachedInfo == nullptr)
-		{
-			ResourceAllocationInfoMap.Add(Hash, Result);
-		}
-		ResourceAllocationInfoMapMutex.WriteUnlock();
-		return Result;
+	ResourceAllocationInfoMapMutex.ReadUnlock();
+
+	Result = GetResourceAllocationInfoUncached(InDesc);
+
+	ResourceAllocationInfoMapMutex.WriteLock();
+	// Try search again with write lock because could have been added already
+	CachedInfo = ResourceAllocationInfoMap.Find(Hash);
+	if (CachedInfo == nullptr)
+	{
+		ResourceAllocationInfoMap.Add(Hash, Result);
 	}
+	ResourceAllocationInfoMapMutex.WriteUnlock();
+
+	return Result;
 }
 
 FD3D12ContextCommon* FD3D12Device::ObtainContext(ED3D12QueueType QueueType)
