Expose D3DDevice9 to avoid needing to cast in code that needs it.

diff --git a/Engine/Source/Runtime/D3D12RHI/Private/D3D12Device.cpp b/Engine/Source/Runtime/D3D12RHI/Private/D3D12Device.cpp
--- a/Engine/Source/Runtime/D3D12RHI/Private/D3D12Device.cpp
+++ b/Engine/Source/Runtime/D3D12RHI/Private/D3D12Device.cpp
@@ -188,6 +188,11 @@ ID3D12Device7* FD3D12Device::GetDevice7()
 {
 	return GetParentAdapter()->GetD3DDevice7();
 }
+
+ID3D12Device9* FD3D12Device::GetDevice9()
+{
+	return GetParentAdapter()->GetD3DDevice9();
+}
 #endif // D3D12_RHI_RAYTRACING
 
 #if D3D12_SUPPORTS_DXGI_DEBUG
diff --git a/Engine/Source/Runtime/D3D12RHI/Private/D3D12Device.h b/Engine/Source/Runtime/D3D12RHI/Private/D3D12Device.h
--- a/Engine/Source/Runtime/D3D12RHI/Private/D3D12Device.h
+++ b/Engine/Source/Runtime/D3D12RHI/Private/D3D12Device.h
@@ -212,6 +212,7 @@ public:
 
 	ID3D12Device5*							  GetDevice5();
 	ID3D12Device7*							  GetDevice7();
+	ID3D12Device9*							  GetDevice9();
 
 	FD3D12RayTracingPipelineCache*			  GetRayTracingPipelineCache           () { return RayTracingPipelineCache;            }
 	FD3D12Buffer*							  GetRayTracingDispatchRaysDescBuffer  (ED3D12QueueType QueueType) { return Queues[(uint32)QueueType].RayTracingDispatchRaysDescBuffer; }
