RHI DX12 Transition/Barrier Refactor

Remove the resource state tracking at the DX12 context level, relying on transition coming from the high level RHI rendering code.
After this refactor the DX12 context will no longer track resource state per commandlist (D3D12_RESOURCE_STATE_TBD is no more a valid state).
FD3D12Resource resource state needs to be in sync all the time with FRHIViewableResource state since the transition will be managed only from platform agnostic RHI code.

diff --git a/Engine/Source/Runtime/D3D12RHI/Private/D3D12Adapter.h b/Engine/Source/Runtime/D3D12RHI/Private/D3D12Adapter.h
--- a/Engine/Source/Runtime/D3D12RHI/Private/D3D12Adapter.h
+++ b/Engine/Source/Runtime/D3D12RHI/Private/D3D12Adapter.h
@@ -383,6 +383,7 @@ public:
 		const FRHIGPUMask& InGPUMask,
 		ID3D12ResourceAllocator* ResourceAllocator,
 		const TCHAR* InDebugName,
+		ERHIAccess InRHICreationState,
 		const FName& InOwnerName = NAME_None,
 		const FName& InClassName = NAME_None);
 
diff --git a/Engine/Source/Runtime/D3D12RHI/Private/D3D12Buffer.cpp b/Engine/Source/Runtime/D3D12RHI/Private/D3D12Buffer.cpp
--- a/Engine/Source/Runtime/D3D12RHI/Private/D3D12Buffer.cpp
+++ b/Engine/Source/Runtime/D3D12RHI/Private/D3D12Buffer.cpp
@@ -238,6 +238,7 @@ FD3D12Buffer* FD3D12Adapter::CreateRHIBuffer(
 	const FRHIGPUMask& InGPUMask,
 	ID3D12ResourceAllocator* ResourceAllocator,
 	const TCHAR* InDebugName,
+	ERHIAccess InRHICreationState,
 	const FName& OwnerName,
 	const FName& ClassName)
 {
@@ -258,6 +259,7 @@ FD3D12Buffer* FD3D12Adapter::CreateRHIBuffer(
 		BufferOut = CreateLinkedObject<FD3D12Buffer>(InGPUMask, [&](FD3D12Device* Device, FD3D12Buffer* FirstLinkedObject)
 		{
 			FD3D12Buffer* NewBuffer = new FD3D12Buffer(Device, BufferDesc);
+			NewBuffer->SetTrackedAccess_Unsafe(InRHICreationState);
 			NewBuffer->BufferAlignment = Alignment;
 
 #if NAME_OBJECTS
@@ -287,6 +289,7 @@ FD3D12Buffer* FD3D12Adapter::CreateRHIBuffer(
 		BufferOut = CreateLinkedObject<FD3D12Buffer>(InGPUMask, [&](FD3D12Device* Device, FD3D12Buffer* FirstLinkedObject)
 		{
 			FD3D12Buffer* NewBuffer = new FD3D12Buffer(Device, BufferDesc);
+			NewBuffer->SetTrackedAccess_Unsafe(InRHICreationState);
 			NewBuffer->BufferAlignment = Alignment;
 
 #if NAME_OBJECTS
@@ -478,7 +481,7 @@ FD3D12Buffer* FD3D12DynamicRHI::CreateD3D12Buffer(class FRHICommandListBase* RHI
 	// which may be added in the future.
 	FRHIGPUMask GPUMask = EnumHasAnyFlags(BufferDesc.Usage, EBufferUsageFlags::NNE) ? FRHIGPUMask::GPU0() : CreateInfo.GPUMask;
 
-	FD3D12Buffer* Buffer = GetAdapter().CreateRHIBuffer(Desc, Alignment, BufferDesc, StateMode, CreateState, bHasInitialData || bForceKeepUnlocked, GPUMask, ResourceAllocator, CreateInfo.DebugName, CreateInfo.OwnerName, TraceClassName);
+	FD3D12Buffer* Buffer = GetAdapter().CreateRHIBuffer(Desc, Alignment, BufferDesc, StateMode, CreateState, bHasInitialData || bForceKeepUnlocked, GPUMask, ResourceAllocator, CreateInfo.DebugName, InResourceState, CreateInfo.OwnerName, TraceClassName);
 	check(Buffer->ResourceLocation.IsValid());
 
 	// Copy the resource data if available 
@@ -690,8 +693,22 @@ void* FD3D12DynamicRHI::RHILockBuffer(FRHICommandListBase& RHICmdList, FRHIBuffe
 	// tracing acceleration structures.
 	check(!EnumHasAnyFlags(BufferRHI->GetUsage(), BUF_MultiGPUAllocate));
 
+	bool bNeedTransition = ( (!(EnumHasAnyFlags(BufferRHI->GetUsage(), BUF_AnyDynamic))) && (LockMode == EResourceLockMode::RLM_ReadOnly) );
+
+	if (RHICmdList.NeedsExtraTransitions() && bNeedTransition)
+	{
+		RHICmdList.TransitionInternal(FRHITransitionInfo(BufferRHI, ERHIAccess::Unknown, ERHIAccess::CopySrc, EResourceTransitionFlags::IgnoreAfterState), ERHITransitionCreateFlags::AllowDuringRenderPass);
+	}
+
 	FD3D12Buffer* Buffer = FD3D12DynamicRHI::ResourceCast(BufferRHI);
-	return LockBuffer(RHICmdList, Buffer, Buffer->GetSize(), Buffer->GetUsage(), Offset, Size, LockMode);
+	void* retval = LockBuffer(RHICmdList, Buffer, Buffer->GetSize(), Buffer->GetUsage(), Offset, Size, LockMode);
+
+	if (RHICmdList.NeedsExtraTransitions() && bNeedTransition)
+	{
+		RHICmdList.TransitionInternal(FRHITransitionInfo(BufferRHI, ERHIAccess::CopySrc, ERHIAccess::Unknown, EResourceTransitionFlags::IgnoreAfterState), ERHITransitionCreateFlags::AllowDuringRenderPass);
+	}
+
+	return retval;
 }
 
 void* FD3D12DynamicRHI::RHILockBufferMGPU(FRHICommandListBase& RHICmdList, FRHIBuffer* BufferRHI, uint32 GPUIndex, uint32 Offset, uint32 Size, EResourceLockMode LockMode)
@@ -710,7 +727,21 @@ void FD3D12DynamicRHI::RHIUnlockBuffer(FRHICommandListBase& RHICmdList, FRHIBuff
 	check(!EnumHasAnyFlags(BufferRHI->GetUsage(), BUF_MultiGPUAllocate));
 
 	FD3D12Buffer* Buffer = FD3D12DynamicRHI::ResourceCast(BufferRHI);
+	FD3D12LockedResource& LockedData = Buffer->LockedData;
+
+	const bool bNeedTransition = (!((EnumHasAnyFlags(Buffer->GetUsage(), BUF_AnyDynamic)) || LockedData.bLockedForReadOnly));
+
+	if (RHICmdList.NeedsExtraTransitions() && bNeedTransition)
+	{
+		RHICmdList.TransitionInternal(FRHITransitionInfo(Buffer, ERHIAccess::Unknown, ERHIAccess::CopyDest, EResourceTransitionFlags::IgnoreAfterState), ERHITransitionCreateFlags::AllowDuringRenderPass);
+	}
+
 	UnlockBuffer(RHICmdList, Buffer, Buffer->GetUsage());
+
+	if (RHICmdList.NeedsExtraTransitions() && bNeedTransition )
+	{
+		RHICmdList.TransitionInternal(FRHITransitionInfo(Buffer, ERHIAccess::CopyDest, ERHIAccess::Unknown, EResourceTransitionFlags::IgnoreAfterState), ERHITransitionCreateFlags::AllowDuringRenderPass);
+	}
 }
 
 void FD3D12DynamicRHI::RHIUnlockBufferMGPU(FRHICommandListBase& RHICmdList, FRHIBuffer* BufferRHI, uint32 GPUIndex)
diff --git a/Engine/Source/Runtime/D3D12RHI/Private/D3D12CommandContext.cpp b/Engine/Source/Runtime/D3D12RHI/Private/D3D12CommandContext.cpp
--- a/Engine/Source/Runtime/D3D12RHI/Private/D3D12CommandContext.cpp
+++ b/Engine/Source/Runtime/D3D12RHI/Private/D3D12CommandContext.cpp
@@ -92,6 +92,8 @@ static D3D12_RESOURCE_STATES GetValidResourceStates(ED3D12QueueType CommandListT
 	return ~DisallowedDirectStates;
 }
 
+D3D12_RESOURCE_STATES FD3D12CommandContext::SkipFastClearEliminateState = D3D12_RESOURCE_STATES(0);
+
 FD3D12CommandContext::FD3D12CommandContext(FD3D12Device* InParent, ED3D12QueueType QueueType, bool InIsDefaultContext)
 	: FD3D12ContextCommon(InParent, QueueType, InIsDefaultContext)
 	, FD3D12CommandContextBase(InParent->GetParentAdapter(), InParent->GetGPUMask())
@@ -108,6 +110,7 @@ FD3D12CommandContext::FD3D12CommandContext(FD3D12Device* InParent, ED3D12QueueTy
 		FD3D12ConstantBuffer(InParent, ConstantsAllocator),
 	}
 {
+	SkipFastClearEliminateState = D3D12_RESOURCE_STATES(0);
 	StaticUniformBuffers.AddZeroed(FUniformBufferStaticSlotRegistry::Get().GetSlotCount());
 	ClearState();
 }
@@ -285,7 +288,7 @@ FD3D12ContextCommon::FD3D12ContextCommon(FD3D12Device* Device, ED3D12QueueType Q
 
 void FD3D12ContextCommon::WaitSyncPoint(FD3D12SyncPoint* SyncPoint)
 {
-	if (IsOpen())
+	if (IsPendingCommands())
 	{
 		CloseCommandList();
 	}
@@ -295,7 +298,7 @@ void FD3D12ContextCommon::WaitSyncPoint(FD3D12SyncPoint* SyncPoint)
 
 void FD3D12ContextCommon::SignalSyncPoint(FD3D12SyncPoint* SyncPoint)
 {
-	if (IsOpen())
+	if (IsPendingCommands())
 	{
 		CloseCommandList();
 	}
@@ -305,7 +308,7 @@ void FD3D12ContextCommon::SignalSyncPoint(FD3D12SyncPoint* SyncPoint)
 
 void FD3D12ContextCommon::SignalManualFence(ID3D12Fence* Fence, uint64 Value)
 {
-	if (IsOpen())
+	if (IsPendingCommands())
 	{
 		CloseCommandList();
 	}
@@ -315,7 +318,7 @@ void FD3D12ContextCommon::SignalManualFence(ID3D12Fence* Fence, uint64 Value)
 
 void FD3D12ContextCommon::WaitManualFence(ID3D12Fence* Fence, uint64 Value)
 {
-	if (IsOpen())
+	if (IsPendingCommands())
 	{
 		CloseCommandList();
 	}
@@ -368,7 +371,7 @@ FD3D12QueryLocation FD3D12ContextCommon::InsertTimestamp(ED3D12Units Units, uint
 
 void FD3D12ContextCommon::SetReservedBufferCommitSize(FD3D12Buffer* Buffer, uint64 CommitSizeInBytes)
 {
-	if (IsOpen())
+	if (IsPendingCommands())
 	{
 		CloseCommandList();
 	}
@@ -411,6 +414,11 @@ void FD3D12CommandContext::OpenCommandList()
 
 void FD3D12ContextCommon::CloseCommandList()
 {
+	checkf(IsPendingCommands(), TEXT("The command list is empty."));
+	// Do this before we insert the final timestamp to ensure we're timing all the work on the command list.
+	// If the command list only has barrier work to do, this will open the command list for the first time
+	FlushResourceBarriers();
+
 	checkf(IsOpen(), TEXT("Command list is not open."));
 	checkf(Payloads.Num() && CurrentPhase == EPhase::Execute, TEXT("Expected the current payload to be in the execute phase."));
 	
@@ -418,9 +426,6 @@ void FD3D12ContextCommon::CloseCommandList()
 
 	FD3D12Payload* Payload = GetPayload(EPhase::Execute);
 
-	// Do this before we insert the final timestamp to ensure we're timing all the work on the command list.
-	FlushResourceBarriers();
-
 	CommandList->Close();
 	CommandList = nullptr;
 
@@ -444,7 +449,7 @@ void FD3D12CommandContext::CloseCommandList()
 
 void FD3D12ContextCommon::Finalize(TArray<FD3D12Payload*>& OutPayloads)
 {
-	if (IsOpen())
+	if (IsPendingCommands())
 	{
 		CloseCommandList();
 	}
@@ -608,7 +613,7 @@ void FD3D12ContextCommon::FlushCommands(ED3D12FlushFlags FlushFlags)
 	// We should only be flushing the default context
 	check(IsDefaultContext());
 
-	if (IsOpen())
+	if (IsPendingCommands())
 	{
 		CloseCommandList();
 	}
@@ -787,7 +792,27 @@ void FD3D12DynamicRHI::RHICreateTransition(FRHITransition* Transition, const FRH
 	Data->CreateFlags = CreateInfo.Flags;
 
 	const bool bCrossPipeline = (CreateInfo.SrcPipelines != CreateInfo.DstPipelines) && (!EnumHasAnyFlags(Data->CreateFlags, ERHITransitionCreateFlags::NoFence));
-	if (bCrossPipeline)
+	const bool bAsyncToAllPipelines = ((CreateInfo.SrcPipelines == ERHIPipeline::AsyncCompute) && (CreateInfo.DstPipelines == ERHIPipeline::All));
+
+	Data->bCrossPipeline = bCrossPipeline;
+
+	// In DX12 we cannot perform resource barrier with graphics state on the AsyncCompute pipe
+	// this check is here to be able to force a crosspipe transition coming from AsyncCompute with graphics states to be processed in the Graphics pipe
+	if (bAsyncToAllPipelines)
+	{
+		for (const FRHITransitionInfo& TransitionInfo : CreateInfo.TransitionInfos)
+		{
+			if (EnumHasAnyFlags(TransitionInfo.AccessAfter, ERHIAccess::SRVGraphics | ERHIAccess::UAVGraphics))
+			{
+				Data->bAsyncToAllPipelines = true;
+				// note that we are ignoring the ERHITransitionCreateFlags::NoFence flag for this specific case and forcing the fences creation, since we need fences between the pipes
+				Data->bCrossPipeline = true;
+				break;
+			}
+		}
+	}
+
+	if (Data->bCrossPipeline)
 	{
 		// Create one sync point per device, per source pipe
 		for (uint32 Index : FRHIGPUMask::All())
@@ -800,7 +825,6 @@ void FD3D12DynamicRHI::RHICreateTransition(FRHITransition* Transition, const FRH
 		}
 	}
 
-	Data->bCrossPipeline = bCrossPipeline;
 	Data->TransitionInfos = CreateInfo.TransitionInfos;
 	Data->AliasingInfos = CreateInfo.AliasingInfos;
 
@@ -940,9 +964,6 @@ void FD3D12CommandContextRedirector::RHITransferResources(TConstArrayView<FTrans
 				SrcResource = FD3D12DynamicRHI::ResourceCast(Param.Buffer.GetReference(), Param.SrcGPUIndex )->GetResource();
 				DstResource = FD3D12DynamicRHI::ResourceCast(Param.Buffer.GetReference(), Param.DestGPUIndex)->GetResource();
 			}
-
-			SrcContext->TransitionResource(SrcResource, D3D12_RESOURCE_STATE_TBD, D3D12_RESOURCE_STATE_COPY_SOURCE, 0);
-			DstContext->TransitionResource(DstResource, D3D12_RESOURCE_STATE_TBD, D3D12_RESOURCE_STATE_COPY_DEST  , 0);
 		}
 	}
 
@@ -1133,31 +1154,6 @@ void FD3D12CommandContextRedirector::RHICrossGPUTransfer(TConstArrayView<FTransf
 	if (Params.Num() == 0)
 		return;
 
-	for (const FTransferResourceParams& Param : Params)
-	{
-		FD3D12CommandContext* SrcContext = PhysicalContexts[Param.SrcGPUIndex];
-		check(SrcContext);
-
-		FD3D12Resource* SrcResource;
-
-		if (Param.Texture)
-		{
-			check(Param.Buffer == nullptr);
-
-			SrcResource = FD3D12CommandContext::RetrieveTexture(Param.Texture, Param.SrcGPUIndex )->GetResource();
-		}
-		else
-		{
-			check(Param.Buffer != nullptr);
-
-			SrcResource = FD3D12DynamicRHI::ResourceCast(Param.Buffer.GetReference(), Param.SrcGPUIndex )->GetResource();
-		}
-
-		// Destination GPU resources are transitioned in RHICrossGPUTransferSignal, potentially allowing them to be transitioned earlier
-		// in the timeline, reducing the likelihood of the source GPU needing to wait to start the transfer.
-		SrcContext->TransitionResource(SrcResource, D3D12_RESOURCE_STATE_TBD, D3D12_RESOURCE_STATE_COPY_SOURCE, 0);
-	}
-
 	// Wait on any pre-transfer fences first
 	for (FCrossGPUTransferFence* PreTransferSyncPoint : PreTransfer)
 	{
@@ -1226,29 +1222,6 @@ void FD3D12CommandContextRedirector::RHICrossGPUTransfer(TConstArrayView<FTransf
 
 void FD3D12CommandContextRedirector::RHICrossGPUTransferSignal(TConstArrayView<FTransferResourceParams> Params, TConstArrayView<FCrossGPUTransferFence*> PreTransfer)
 {
-	for (const FTransferResourceParams& Param : Params)
-	{
-		FD3D12CommandContext* DstContext = PhysicalContexts[Param.DestGPUIndex];
-		check(DstContext);
-
-		FD3D12Resource* DstResource;
-
-		if (Param.Texture)
-		{
-			check(Param.Buffer == nullptr);
-
-			DstResource = FD3D12CommandContext::RetrieveTexture(Param.Texture, Param.DestGPUIndex)->GetResource();
-		}
-		else
-		{
-			check(Param.Buffer != nullptr);
-
-			DstResource = FD3D12DynamicRHI::ResourceCast(Param.Buffer.GetReference(), Param.DestGPUIndex)->GetResource();
-		}
-
-		DstContext->TransitionResource(DstResource, D3D12_RESOURCE_STATE_TBD, D3D12_RESOURCE_STATE_COPY_DEST, 0);
-	}
-
 	for (FCrossGPUTransferFence* TransferSyncPoint : PreTransfer)
 	{
 		FD3D12SyncPointRef SyncPoint = FD3D12SyncPoint::Create(ED3D12SyncPointType::GPUOnly);
diff --git a/Engine/Source/Runtime/D3D12RHI/Private/D3D12CommandContext.h b/Engine/Source/Runtime/D3D12RHI/Private/D3D12CommandContext.h
--- a/Engine/Source/Runtime/D3D12RHI/Private/D3D12CommandContext.h
+++ b/Engine/Source/Runtime/D3D12RHI/Private/D3D12CommandContext.h
@@ -237,6 +237,7 @@ public:
 	bool IsDefaultContext() const { return bIsDefaultContext; }
 
 	bool IsOpen() const { return CommandList != nullptr; }
+	bool IsPendingCommands() const { return IsOpen() || ResourceBarrierBatcher.Num(); }
 
 	FD3D12SyncPoint* GetContextSyncPoint()
 	{
@@ -394,33 +395,19 @@ public:
 	void UpdateResidency(const FD3D12Resource* Resource) { }
 #endif
 
-	// Pending resource barriers are resolved by a dedicated barrier command list, generated during command list submission.
-	void AddPendingResourceBarrier(FD3D12Resource* Resource, D3D12_RESOURCE_STATES After, uint32 SubResource, CResourceState& ResourceState_OnCommandList);
-
 	// Resource transition / barrier functions. These get batched and recorded into the command list when FlushResourceBarriers() is called.
-	void AddTransitionBarrier(FD3D12Resource* pResource, D3D12_RESOURCE_STATES Before, D3D12_RESOURCE_STATES After, uint32 Subresource, CResourceState* ResourceState_OnCommandList = nullptr);
+	void AddTransitionBarrier(FD3D12Resource* pResource, D3D12_RESOURCE_STATES Before, D3D12_RESOURCE_STATES After, uint32 Subresource);
 	void AddAliasingBarrier(ID3D12Resource* InResourceBefore, ID3D12Resource* InResourceAfter);
 	void AddUAVBarrier();
 
 	// Flushes the batched resource barriers to the current command list
 	void FlushResourceBarriers();
 
-	// Helper functions for transitioning resource views when the previous state is unknown.
-	// Uses a combination of pending barriers and command list state tracking to resolve the unknown previous state.
-	void TransitionResource(class FD3D12DepthStencilView*    View);
-	void TransitionResource(class FD3D12DepthStencilView*    View, D3D12_RESOURCE_STATES After);
-	void TransitionResource(class FD3D12RenderTargetView*    View, D3D12_RESOURCE_STATES After);
-	void TransitionResource(class FD3D12ShaderResourceView*  View, D3D12_RESOURCE_STATES After);
-	void TransitionResource(class FD3D12UnorderedAccessView* View, D3D12_RESOURCE_STATES After);
-
-	// Functions for transitioning a resource. The Before state may be D3D12_RESOURCE_STATE_TBD, in which case a
-	// combination of pending barriers and command list state tracking is used to resolve the unknown previous state.
+	// Functions for transitioning a resource. The Before and After state cannot be D3D12_RESOURCE_STATE_TBD, needs be known at the time of the call
 	bool TransitionResource(FD3D12Resource* Resource, D3D12_RESOURCE_STATES Before, D3D12_RESOURCE_STATES After, uint32 Subresource);
-	bool TransitionResource(FD3D12Resource* Resource, D3D12_RESOURCE_STATES Before, D3D12_RESOURCE_STATES After, FD3D12ViewSubset const& ViewSubset);
-
+	
 private:
-	// Private helper for transitioning resources
-	bool TransitionResource(FD3D12Resource* InResource, CResourceState& ResourceState_OnCommandList, uint32 InSubresourceIndex, D3D12_RESOURCE_STATES InBeforeState, D3D12_RESOURCE_STATES InAfterState, bool bInForceAfterState);
+	bool TransitionResource(FD3D12Resource* InResource, uint32 InSubresourceIndex, D3D12_RESOURCE_STATES InBeforeState, D3D12_RESOURCE_STATES InAfterState);
 };
 
 //
@@ -569,7 +556,7 @@ public:
 	D3D12_SHADING_RATE				VRSShadingRate = D3D12_SHADING_RATE_1X1;
 #endif
 
-	D3D12_RESOURCE_STATES SkipFastClearEliminateState = D3D12_RESOURCE_STATES(0);
+	static D3D12_RESOURCE_STATES SkipFastClearEliminateState;
 	D3D12_RESOURCE_STATES ValidResourceStates;
 
 	/** Constant buffers for Set*ShaderParameter calls. */
@@ -1177,6 +1164,7 @@ struct FD3D12TransitionData
 	TArray<TRHIPipelineArray<FD3D12SyncPointRef>, TInlineAllocator<MAX_NUM_GPUS>> SyncPoints;
 
 	bool bCrossPipeline = false;
+	bool bAsyncToAllPipelines = false;
 };
 
 class FD3D12ContextArray : public TRHIPipelineArray<FD3D12CommandContextBase*>
diff --git a/Engine/Source/Runtime/D3D12RHI/Private/D3D12CommandList.cpp b/Engine/Source/Runtime/D3D12RHI/Private/D3D12CommandList.cpp
--- a/Engine/Source/Runtime/D3D12RHI/Private/D3D12CommandList.cpp
+++ b/Engine/Source/Runtime/D3D12RHI/Private/D3D12CommandList.cpp
@@ -59,17 +59,8 @@ void FD3D12CommandList::AddToResidencySet(TConstArrayView<FD3D12ResidencyHandle*
 }
 #endif // ENABLE_RESIDENCY_MANAGEMENT
 
-void FD3D12ContextCommon::AddPendingResourceBarrier(FD3D12Resource* Resource, D3D12_RESOURCE_STATES After, uint32 SubResource, CResourceState& ResourceState_OnCommandList)
-{
-	check(After != D3D12_RESOURCE_STATE_TBD);
-	check(Resource->RequiresResourceStateTracking());
-	check(&GetCommandList().GetResourceState_OnCommandList(Resource) == &ResourceState_OnCommandList);
-	
-	GetCommandList().State.PendingResourceBarriers.Emplace(Resource, After, SubResource);
-	ResourceState_OnCommandList.SetSubresourceState(SubResource, After);	
-}
 
-void FD3D12ContextCommon::AddTransitionBarrier(FD3D12Resource* pResource, D3D12_RESOURCE_STATES Before, D3D12_RESOURCE_STATES After, uint32 Subresource, CResourceState* ResourceState_OnCommandList)
+void FD3D12ContextCommon::AddTransitionBarrier(FD3D12Resource* pResource, D3D12_RESOURCE_STATES Before, D3D12_RESOURCE_STATES After, uint32 Subresource)
 {
 	if (Before != After)
 	{
@@ -86,22 +77,6 @@ void FD3D12ContextCommon::AddTransitionBarrier(FD3D12Resource* pResource, D3D12_
 	{
 		ensureMsgf(0, TEXT("AddTransitionBarrier: Before == After (%d)"), (uint32)Before);
 	}
-
-	if (pResource->RequiresResourceStateTracking())
-	{
-		if (!ResourceState_OnCommandList)
-		{
-			ResourceState_OnCommandList = &GetCommandList().GetResourceState_OnCommandList(pResource);
-		}
-		else
-		{
-			// If a resource state was passed to avoid a repeat lookup, it must be the one we expected.
-			check(&GetCommandList().GetResourceState_OnCommandList(pResource) == ResourceState_OnCommandList);
-		}
-
-		ResourceState_OnCommandList->SetSubresourceState(Subresource, After);
-		ResourceState_OnCommandList->SetHasInternalTransition();
-	}
 }
 
 void FD3D12ContextCommon::AddUAVBarrier()
@@ -133,26 +108,6 @@ void FD3D12ContextCommon::FlushResourceBarriers()
 	}
 }
 
-CResourceState& FD3D12CommandList::GetResourceState_OnCommandList(FD3D12Resource* pResource)
-{
-	// Only certain resources should use this
-	check(pResource->RequiresResourceStateTracking());
-
-	CResourceState& ResourceState = State.TrackedResourceState.FindOrAdd(pResource);
-
-	// If there is no entry, all subresources should be in the resource's TBD state.
-	// This means we need to have pending resource barrier(s).
-	if (!ResourceState.CheckResourceStateInitalized())
-	{
-		ResourceState.Initialize(pResource->GetSubresourceCount());
-		check(ResourceState.CheckResourceState(D3D12_RESOURCE_STATE_TBD));
-	}
-
-	check(ResourceState.CheckResourceStateInitalized());
-
-	return ResourceState;
-}
-
 FD3D12CommandAllocator::FD3D12CommandAllocator(FD3D12Device* Device, ED3D12QueueType QueueType)
 	: Device(Device)
 	, QueueType(QueueType)
@@ -298,15 +253,24 @@ void FD3D12CommandList::Close()
 	check(IsOpen());
 	EndLocalQueries();
 
+	HRESULT hr;
 	if (Interfaces.CopyCommandList)
 	{
-		VERIFYD3D12RESULT(Interfaces.CopyCommandList->Close());
+		hr = Interfaces.CopyCommandList->Close();
 	}
 	else
 	{
-		VERIFYD3D12RESULT(Interfaces.GraphicsCommandList->Close());
+		hr = Interfaces.GraphicsCommandList->Close();
 	}
 
+#if DEBUG_RESOURCE_STATES
+	#define DX12_RESOURCE_NAME_TO_LOG	"Lumen.RadianceCache.DepthProbeAtlasTexture"
+	if (hr != S_OK)
+		LogResourceBarriers(State.ResourceBarriers, Interfaces.CommandList.GetReference(), FString(DX12_RESOURCE_NAME_TO_LOG));
+#endif
+
+	VERIFYD3D12RESULT(hr);
+
 	if (State.DeferredResidencyUpdateSet.Num() == 0)
 	{
 		D3DX12Residency::Close(ResidencySet);
@@ -461,8 +425,6 @@ void FD3D12CommandList::WriteTimestamp(FD3D12QueryLocation const& Location, ED3D
 FD3D12CommandList::FState::FState(FD3D12CommandAllocator* CommandAllocator, FD3D12QueryAllocator* TimestampAllocator, FD3D12QueryAllocator* PipelineStatsAllocator)
 	: CommandAllocator(CommandAllocator)
 {
-	PendingResourceBarriers.Reserve(256);
-
 	if (TimestampAllocator)
 	{
 #if RHI_NEW_GPU_PROFILER
@@ -480,270 +442,50 @@ FD3D12CommandList::FState::FState(FD3D12CommandAllocator* CommandAllocator, FD3D
 	}
 }
 
-void FD3D12ContextCommon::TransitionResource(FD3D12DepthStencilView* View)
-{
-	// Determine the required subresource states from the view desc
-	const D3D12_DEPTH_STENCIL_VIEW_DESC& DSVDesc = View->GetD3DDesc();
-	const bool bDSVDepthIsWritable = (DSVDesc.Flags & D3D12_DSV_FLAG_READ_ONLY_DEPTH) == 0;
-	const bool bDSVStencilIsWritable = (DSVDesc.Flags & D3D12_DSV_FLAG_READ_ONLY_STENCIL) == 0;
-	// TODO: Check if the PSO depth stencil is writable. When this is done, we need to transition in SetDepthStencilState too.
-
-	// This code assumes that the DSV always contains the depth plane
-	check(View->HasDepth());
-	const bool bHasDepth = true;
-	const bool bHasStencil = View->HasStencil();
-	const bool bDepthIsWritable = bHasDepth && bDSVDepthIsWritable;
-	const bool bStencilIsWritable = bHasStencil && bDSVStencilIsWritable;
-
-	// DEPTH_WRITE is suitable for read operations when used as a normal depth/stencil buffer.
-	FD3D12Resource* Resource = View->GetResource();
-	if (bDepthIsWritable)
-	{
-		TransitionResource(Resource, D3D12_RESOURCE_STATE_TBD, D3D12_RESOURCE_STATE_DEPTH_WRITE, View->GetDepthOnlySubset());
-	}
-	else if (bHasDepth && GD3D12ExtraDepthTransitions)
-	{
-		TransitionResource(Resource, D3D12_RESOURCE_STATE_TBD, D3D12_RESOURCE_STATE_DEPTH_READ, View->GetDepthOnlySubset());
-	}
-
-	if (bStencilIsWritable)
-	{
-		TransitionResource(Resource, D3D12_RESOURCE_STATE_TBD, D3D12_RESOURCE_STATE_DEPTH_WRITE, View->GetStencilOnlySubset());
-	}
-	else if (bHasStencil && GD3D12ExtraDepthTransitions)
-	{
-		TransitionResource(Resource, D3D12_RESOURCE_STATE_TBD, D3D12_RESOURCE_STATE_DEPTH_READ, View->GetStencilOnlySubset());
-	}
-}
-
-void FD3D12ContextCommon::TransitionResource(FD3D12DepthStencilView* View, D3D12_RESOURCE_STATES After)
-{
-	FD3D12Resource* Resource = View->GetResource();
-
-	const D3D12_DEPTH_STENCIL_VIEW_DESC& Desc = View->GetD3DDesc();
-	switch (Desc.ViewDimension)
-	{
-	case D3D12_DSV_DIMENSION_TEXTURE2D:
-	case D3D12_DSV_DIMENSION_TEXTURE2DMS:
-		if (Resource->GetPlaneCount() > 1)
-		{
-			// Multiple subresources to transtion
-			TransitionResource(Resource, D3D12_RESOURCE_STATE_TBD, After, View->GetViewSubset());
-		}
-		else
-		{
-			// Only one subresource to transition
-			check(Resource->GetPlaneCount() == 1);
-			TransitionResource(Resource, D3D12_RESOURCE_STATE_TBD, After, Desc.Texture2D.MipSlice);
-		}
-		break;
-
-	case D3D12_DSV_DIMENSION_TEXTURE2DARRAY:
-	case D3D12_DSV_DIMENSION_TEXTURE2DMSARRAY:
-		// Multiple subresources to transtion
-		TransitionResource(Resource, D3D12_RESOURCE_STATE_TBD, After, View->GetViewSubset());
-		break;
-
-	default:
-		checkNoEntry();
-		break;
-	}
-}
-
-void FD3D12ContextCommon::TransitionResource(FD3D12RenderTargetView* View, D3D12_RESOURCE_STATES After)
-{
-	FD3D12Resource* Resource = View->GetResource();
-
-	const D3D12_RENDER_TARGET_VIEW_DESC& Desc = View->GetD3DDesc();
-	switch (Desc.ViewDimension)
-	{
-	case D3D12_RTV_DIMENSION_TEXTURE3D:
-		// Note: For volume (3D) textures, all slices for a given mipmap level are a single subresource index.
-		// Fall-through
-	case D3D12_RTV_DIMENSION_TEXTURE2D:
-	case D3D12_RTV_DIMENSION_TEXTURE2DMS:
-		// Only one subresource to transition
-		TransitionResource(Resource, D3D12_RESOURCE_STATE_TBD, After, Desc.Texture2D.MipSlice);
-		break;
-
-	case D3D12_RTV_DIMENSION_TEXTURE2DARRAY:
-	case D3D12_RTV_DIMENSION_TEXTURE2DMSARRAY:
-		// Multiple subresources to transition
-		TransitionResource(Resource, D3D12_RESOURCE_STATE_TBD, After, View->GetViewSubset());
-		break;
-
-	default:
-		UE_LOG(LogRHI, Fatal, TEXT("Unsupported resource dimension %d"), Desc.ViewDimension);
-		break;
-	}
-}
-
-void FD3D12ContextCommon::TransitionResource(FD3D12ShaderResourceView* View, D3D12_RESOURCE_STATES After)
-{
-	FD3D12Resource* Resource = View->GetResource();
-
-	// Early out if we never need to do state tracking, the resource should always be in an SRV state.
-	if (!Resource || !Resource->RequiresResourceStateTracking())
-		return;
-
-	const D3D12_RESOURCE_DESC& ResDesc = Resource->GetDesc();
-	const FD3D12ViewSubset& ViewSubset = View->GetViewSubset();
-
-	const D3D12_SHADER_RESOURCE_VIEW_DESC& Desc = View->GetD3DDesc();
-	switch (Desc.ViewDimension)
-	{
-	default:
-		// Transition the resource
-		TransitionResource(Resource, D3D12_RESOURCE_STATE_TBD, After, ViewSubset);
-		break;
-
-	case D3D12_SRV_DIMENSION_BUFFER:
-		if (Resource->GetHeapType() == D3D12_HEAP_TYPE_DEFAULT)
-		{
-			// Transition the resource
-			TransitionResource(Resource, D3D12_RESOURCE_STATE_TBD, After, ViewSubset);
-		}
-		break;
-	}
-}
-
-void FD3D12ContextCommon::TransitionResource(FD3D12UnorderedAccessView* View, D3D12_RESOURCE_STATES After)
-{
-	FD3D12Resource* Resource = View->GetResource();
-
-	const D3D12_UNORDERED_ACCESS_VIEW_DESC& Desc = View->GetD3DDesc();
-	switch (Desc.ViewDimension)
-	{
-	case D3D12_UAV_DIMENSION_BUFFER:
-		TransitionResource(Resource, D3D12_RESOURCE_STATE_TBD, After, 0);
-		break;
-
-	case D3D12_UAV_DIMENSION_TEXTURE2D:
-		// Only one subresource to transition
-		TransitionResource(Resource, D3D12_RESOURCE_STATE_TBD, After, Desc.Texture2D.MipSlice);
-		break;
-
-	case D3D12_UAV_DIMENSION_TEXTURE2DARRAY:
-		// Multiple subresources to transtion
-		TransitionResource(Resource, D3D12_RESOURCE_STATE_TBD, After, View->GetViewSubset());
-		break;
-
-	case D3D12_UAV_DIMENSION_TEXTURE3D:
-		// Multiple subresources to transtion
-		TransitionResource(Resource, D3D12_RESOURCE_STATE_TBD, After, View->GetViewSubset());
-		break;
-
-	default:
-		checkNoEntry();
-		break;
-	}
-}
-
 bool FD3D12ContextCommon::TransitionResource(FD3D12Resource* Resource, D3D12_RESOURCE_STATES Before, D3D12_RESOURCE_STATES After, uint32 Subresource)
 {
 	check(Resource);
 	check(Resource->RequiresResourceStateTracking());
 	check(!((After & (D3D12_RESOURCE_STATE_PIXEL_SHADER_RESOURCE | D3D12_RESOURCE_STATE_NON_PIXEL_SHADER_RESOURCE)) && (Resource->GetDesc().Flags & D3D12_RESOURCE_FLAG_DENY_SHADER_RESOURCE)));
+	check(Before != D3D12_RESOURCE_STATE_TBD);
+	check(After != D3D12_RESOURCE_STATE_TBD);
 
 #ifdef PLATFORM_SUPPORTS_RESOURCE_COMPRESSION
 	After |= Resource->GetCompressedState();
 #endif
-
-	UpdateResidency(Resource);
-
-	bool bRequireUAVBarrier = false;
-
-	CResourceState& ResourceState = GetCommandList().GetResourceState_OnCommandList(Resource);
-	if (Subresource == D3D12_RESOURCE_BARRIER_ALL_SUBRESOURCES && !ResourceState.AreAllSubresourcesSame())
+#if DEBUG_RESOURCE_STATES
+	FString IncompatibilityReason;
+	if (CheckResourceStateCompatibility(After, Resource->GetDesc().Flags, IncompatibilityReason) == false)
 	{
-		// Slow path. Want to transition the entire resource (with multiple subresources). But they aren't in the same state.
-
-		const uint32 SubresourceCount = Resource->GetSubresourceCount();
-		for (uint32 SubresourceIndex = 0; SubresourceIndex < SubresourceCount; SubresourceIndex++)
-		{
-			bool bForceInAfterState = true;
-			bRequireUAVBarrier |= TransitionResource(Resource, ResourceState, SubresourceIndex, Before, After, bForceInAfterState);
-		}
-
-		// The entire resource should now be in the after state on this command list (even if all barriers are pending)
-		verify(ResourceState.CheckAllSubresourceSame());
-		check(EnumHasAllFlags(ResourceState.GetSubresourceState(D3D12_RESOURCE_BARRIER_ALL_SUBRESOURCES), After));
+		UE_LOG(LogRHI, Error, TEXT("Incompatible Transition State for Resource %s - %s"), *Resource->GetName().ToString(), *IncompatibilityReason);
 	}
-	else
-	{
-		bool bForceInAfterState = false;
-		bRequireUAVBarrier = TransitionResource(Resource, ResourceState, Subresource, Before, After, bForceInAfterState);
-	}
-
-	return bRequireUAVBarrier;
-}
-
-// Transition subresources from current to a new state, using resource state tracking.
-bool FD3D12ContextCommon::TransitionResource(FD3D12Resource* Resource, D3D12_RESOURCE_STATES Before, D3D12_RESOURCE_STATES After, FD3D12ViewSubset const& ViewSubset)
-{
-	check(Resource);
-	check(Resource->RequiresResourceStateTracking());
-	check(!((After & (D3D12_RESOURCE_STATE_PIXEL_SHADER_RESOURCE | D3D12_RESOURCE_STATE_NON_PIXEL_SHADER_RESOURCE)) && (Resource->GetDesc().Flags & D3D12_RESOURCE_FLAG_DENY_SHADER_RESOURCE)));
-
-#ifdef PLATFORM_SUPPORTS_RESOURCE_COMPRESSION
-	After |= Resource->GetCompressedState();
 #endif
-
 	UpdateResidency(Resource);
 
-	const bool bIsWholeResource = ViewSubset.IsWholeResource();
-	CResourceState& ResourceState = GetCommandList().GetResourceState_OnCommandList(Resource);
-
 	bool bRequireUAVBarrier = false;
 
-	if (bIsWholeResource && ResourceState.AreAllSubresourcesSame())
-	{
-		// Fast path. Transition the entire resource from one state to another.
-		bool bForceInAfterState = false;
-		bRequireUAVBarrier = TransitionResource(Resource, ResourceState, D3D12_RESOURCE_BARRIER_ALL_SUBRESOURCES, Before, After, bForceInAfterState);
-	}
-	else
-	{
-		// Slower path. Either the subresources are in more than 1 state, or the view only partially covers the resource.
-		// Either way, we'll need to loop over each subresource in the view...
-
-		bool bWholeResourceWasTransitionedToSameState = bIsWholeResource;
-		for (uint32 SubresourceIndex : ViewSubset)
-		{
-			bool bForceInAfterState = false;
-			bRequireUAVBarrier |= TransitionResource(Resource, ResourceState, SubresourceIndex, Before, After, bForceInAfterState);
-
-			// Subresource not in the same state, then whole resource is not in the same state anymore
-			if (ResourceState.GetSubresourceState(SubresourceIndex) != After)
-				bWholeResourceWasTransitionedToSameState = false;
-		}
-
-		// If we just transtioned every subresource to the same state, lets update it's tracking so it's on a per-resource level
-		if (bWholeResourceWasTransitionedToSameState)
-		{
-			// Sanity check to make sure all subresources are really in the 'after' state
-			verify(ResourceState.CheckAllSubresourceSame());
-			check(EnumHasAllFlags(ResourceState.GetSubresourceState(D3D12_RESOURCE_BARRIER_ALL_SUBRESOURCES), After));
-		}
-	}
+	bRequireUAVBarrier = TransitionResource(Resource, Subresource, Before, After);
 
 	return bRequireUAVBarrier;
 }
 
-static inline bool IsTransitionNeeded(D3D12_RESOURCE_STATES Before, D3D12_RESOURCE_STATES& After)
+static inline bool IsTransitionNeeded(D3D12_RESOURCE_STATES Before, D3D12_RESOURCE_STATES& After, FD3D12Resource* InResource = nullptr)
 {
 	check(Before != D3D12_RESOURCE_STATE_CORRUPT && After != D3D12_RESOURCE_STATE_CORRUPT);
 	check(Before != D3D12_RESOURCE_STATE_TBD && After != D3D12_RESOURCE_STATE_TBD);
 
-	// Depth write is actually a suitable for read operations as a "normal" depth buffer.
-	if (Before == D3D12_RESOURCE_STATE_DEPTH_WRITE && After == D3D12_RESOURCE_STATE_DEPTH_READ)
-	{
-		if (GD3D12ExtraDepthTransitions)
-		{
-			After = D3D12_RESOURCE_STATE_DEPTH_WRITE;
-		}
-		return false;
-	}
+	// Todo: verify that high level is taking care of this optimization
+// 	// Depth write is actually a suitable for read operations as a "normal" depth buffer.
+// 	if (Before == D3D12_RESOURCE_STATE_DEPTH_WRITE && After == D3D12_RESOURCE_STATE_DEPTH_READ)
+// 	{
+// 		{
+// 			if (GD3D12ExtraDepthTransitions)
+// 			{
+// 				After = D3D12_RESOURCE_STATE_DEPTH_WRITE;
+// 			}
+// 			return false;
+// 		}
+// 	}
 
 	// COMMON is an oddball state that doesn't follow the RESOURE_STATE pattern of 
 	// having exactly one bit set so we need to special case these
@@ -754,106 +496,65 @@ static inline bool IsTransitionNeeded(D3D12_RESOURCE_STATES Before, D3D12_RESOUR
 		return true;
 	}
 
+	// Todo: verify that high level is taking care of this optimization
 	// We should avoid doing read-to-read state transitions. But when we do, we should avoid turning off already transitioned bits,
 	// e.g. VERTEX_BUFFER -> SHADER_RESOURCE is turned into VERTEX_BUFFER -> VERTEX_BUFFER | SHADER_RESOURCE.
 	// This reduces the number of resource transitions and ensures automatic states from resource bindings get properly combined.
-	D3D12_RESOURCE_STATES Combined = Before | After;
-	if ((Combined & (D3D12_RESOURCE_STATE_GENERIC_READ | D3D12_RESOURCE_STATE_INDIRECT_ARGUMENT)) == Combined)
-	{
-		After = Combined;
-	}
+// 	D3D12_RESOURCE_STATES Combined = Before | After;
+// 	if ((Combined & (D3D12_RESOURCE_STATE_GENERIC_READ | D3D12_RESOURCE_STATE_INDIRECT_ARGUMENT)) == Combined)
+// 	{
+// 		{
+// 			After = Combined;
+// 		}
+// 	}
 
 	return Before != After;
 }
 
-bool FD3D12ContextCommon::TransitionResource(FD3D12Resource* InResource, CResourceState& ResourceState_OnCommandList, uint32 InSubresourceIndex, D3D12_RESOURCE_STATES InBeforeState, D3D12_RESOURCE_STATES InAfterState, bool bInForceAfterState)
+bool FD3D12ContextCommon::TransitionResource(FD3D12Resource* InResource, uint32 InSubresourceIndex, D3D12_RESOURCE_STATES InBeforeState, D3D12_RESOURCE_STATES InAfterState)
 {
-	// Try and get the correct D3D before state for the transition
-	D3D12_RESOURCE_STATES const TrackedState = ResourceState_OnCommandList.GetSubresourceState(InSubresourceIndex);
-	D3D12_RESOURCE_STATES BeforeState = TrackedState != D3D12_RESOURCE_STATE_TBD ? TrackedState : InBeforeState;
+	check(InBeforeState != D3D12_RESOURCE_STATE_TBD);
+	check(InAfterState != D3D12_RESOURCE_STATE_TBD);
 
-	// Make sure the before states match up or are unknown
-	check(InBeforeState == D3D12_RESOURCE_STATE_TBD || BeforeState == InBeforeState);
+	D3D12_RESOURCE_STATES BeforeState = InBeforeState;
+	D3D12_RESOURCE_STATES AfterState = InAfterState;
 
 	bool bRequireUAVBarrier = false;
-	if (BeforeState != D3D12_RESOURCE_STATE_TBD)
-	{
-		bool bApplyTransitionBarrier = true;
-
-		// Require UAV barrier when before and after are UAV
-		if (BeforeState == D3D12_RESOURCE_STATE_UNORDERED_ACCESS && InAfterState == D3D12_RESOURCE_STATE_UNORDERED_ACCESS)
-		{
-			bRequireUAVBarrier = true;
-		}
-		// Special case for UAV access resources
-		else if (InResource->GetUAVAccessResource() && EnumHasAnyFlags(BeforeState | InAfterState, D3D12_RESOURCE_STATE_UNORDERED_ACCESS))
-		{
-			// inject an aliasing barrier
-			const bool bFromUAV = EnumHasAnyFlags(BeforeState, D3D12_RESOURCE_STATE_UNORDERED_ACCESS);
-			const bool bToUAV = EnumHasAnyFlags(InAfterState, D3D12_RESOURCE_STATE_UNORDERED_ACCESS);
-			check(bFromUAV != bToUAV);
-
-			AddAliasingBarrier(
-				bFromUAV ? InResource->GetUAVAccessResource() : InResource->GetResource(),
-				bToUAV ? InResource->GetUAVAccessResource() : InResource->GetResource());
 
-			if (bToUAV)
-			{
-				ResourceState_OnCommandList.SetUAVHiddenResourceState(BeforeState);
-				bApplyTransitionBarrier = false;
-			}
-			else
-			{
-				D3D12_RESOURCE_STATES HiddenState = ResourceState_OnCommandList.GetUAVHiddenResourceState();
-
-				// Still unknown in this command list?
-				if (HiddenState == D3D12_RESOURCE_STATE_TBD)
-				{
-					AddPendingResourceBarrier(InResource, InAfterState, InSubresourceIndex, ResourceState_OnCommandList);
-					bApplyTransitionBarrier = false;
-				}
-				else
-				{
-					// Use the hidden state as the before state on the resource
-					BeforeState = HiddenState;
-				}
-			}
-		}
-
-		if (bApplyTransitionBarrier)
-		{
-			// We're not using IsTransitionNeeded() when bInForceAfterState is set because we do want to transition even if 'after' is a subset of 'before'
-			// This is so that we can ensure all subresources are in the same state, simplifying future barriers
-			// No state merging when using engine transitions - otherwise next before state might not match up anymore)
-			if ((bInForceAfterState && BeforeState != InAfterState) || IsTransitionNeeded(BeforeState, InAfterState))
-			{
-				AddTransitionBarrier(InResource, BeforeState, InAfterState, InSubresourceIndex, &ResourceState_OnCommandList);
-			}
-			else
-			{
-				// Ensure the command list tracking is up to date, even if we skipped an unnecessary transition.
-				ResourceState_OnCommandList.SetSubresourceState(InSubresourceIndex, InAfterState);
-			}
-		}
+	// Require UAV barrier when before and after are UAV
+	if (BeforeState == D3D12_RESOURCE_STATE_UNORDERED_ACCESS && InAfterState == D3D12_RESOURCE_STATE_UNORDERED_ACCESS)
+	{
+		bRequireUAVBarrier = true;
 	}
-	else
+	// Special case for UAV access resources
+	else if (InResource->GetUAVAccessResource() && EnumHasAnyFlags(BeforeState | InAfterState, D3D12_RESOURCE_STATE_UNORDERED_ACCESS))
 	{
-		// BeforeState is TBD. We need a pending resource barrier.
+		// We are issuing an Aliasing Barrier from or to GetUAVAccessResource
+		// while we are transitioning to a known resource state (D3D12_RESOURCE_STATE_PIXEL_SHADER_RESOURCE) the original resource since it can't go int UAV
+		// the transition into D3D12_RESOURCE_STATE_PIXEL_SHADER_RESOURCE is not needed but in this way the tracking of the resource will be consistent at the higher level
+
+		// inject an aliasing barrier
+		const bool bFromUAV = EnumHasAnyFlags(BeforeState, D3D12_RESOURCE_STATE_UNORDERED_ACCESS);
+		const bool bToUAV = EnumHasAnyFlags(InAfterState, D3D12_RESOURCE_STATE_UNORDERED_ACCESS);
+		check(bFromUAV != bToUAV);
 
-		// Special handling for UAVAccessResource and transition to UAV - don't want to enqueue pending resource to UAV because the actual resource won't transition
-		// Adding of patch up will only be done when transitioning to non UAV state
-		if (InResource->GetUAVAccessResource() && EnumHasAnyFlags(InAfterState, D3D12_RESOURCE_STATE_UNORDERED_ACCESS))
+		if (bToUAV)
 		{
 			AddAliasingBarrier(InResource->GetResource(), InResource->GetUAVAccessResource());
-			ResourceState_OnCommandList.SetUAVHiddenResourceState(D3D12_RESOURCE_STATE_TBD);
+			AfterState = D3D12_RESOURCE_STATE_PIXEL_SHADER_RESOURCE; //original resource CAN'T go into D3D12_RESOURCE_STATE_UNORDERED_ACCESS so we transition it in a known state D3D12_RESOURCE_STATE_PIXEL_SHADER_RESOURCE
 		}
 		else
 		{
-			// We need a pending resource barrier so we can setup the state before this command list executes
-			AddPendingResourceBarrier(InResource, InAfterState, InSubresourceIndex, ResourceState_OnCommandList);
+			AddAliasingBarrier(InResource->GetUAVAccessResource(), InResource->GetResource());
+			BeforeState = D3D12_RESOURCE_STATE_PIXEL_SHADER_RESOURCE; //original resource CAN'T be into D3D12_RESOURCE_STATE_UNORDERED_ACCESS so it was transitioned to D3D12_RESOURCE_STATE_PIXEL_SHADER_RESOURCE
 		}
 	}
 
+	if (IsTransitionNeeded(BeforeState, AfterState, InResource))
+	{
+		AddTransitionBarrier(InResource, BeforeState, AfterState, InSubresourceIndex);
+	}
+
 	return bRequireUAVBarrier;
 }
 
diff --git a/Engine/Source/Runtime/D3D12RHI/Private/D3D12CommandList.h b/Engine/Source/Runtime/D3D12RHI/Private/D3D12CommandList.h
--- a/Engine/Source/Runtime/D3D12RHI/Private/D3D12CommandList.h
+++ b/Engine/Source/Runtime/D3D12RHI/Private/D3D12CommandList.h
@@ -98,10 +98,6 @@ private:
 	FD3D12ResidencySet* const ResidencySet;
 
 public:
-	// Get the state of a resource on this command lists.
-	// This is only used for resources that require state tracking.
-	CResourceState& GetResourceState_OnCommandList(FD3D12Resource* pResource);
-
 	// Indicate that a resource must be made resident before execution on GPU.
 	// Either immediately adds residency handle for this resource to the residency set
 	// or defers it until submission time (residency handles may not be known until then).
@@ -288,12 +284,6 @@ private:
 		// The allocator currently assigned to this command list.
 		FD3D12CommandAllocator* CommandAllocator;
 
-		// Array of resources whose state needs to be synced between submits.
-		TArray<FD3D12PendingResourceBarrier> PendingResourceBarriers;
-
-		// A map of all D3D resources, and their states, that were state transitioned with tracking.
-		TMap<FD3D12Resource*, CResourceState> TrackedResourceState;
-
 #if DEBUG_RESOURCE_STATES
 		// Tracks all the resources barriers being issued on this command list in order
 		TArray<D3D12_RESOURCE_BARRIER> ResourceBarriers;
diff --git a/Engine/Source/Runtime/D3D12RHI/Private/D3D12Commands.cpp b/Engine/Source/Runtime/D3D12RHI/Private/D3D12Commands.cpp
--- a/Engine/Source/Runtime/D3D12RHI/Private/D3D12Commands.cpp
+++ b/Engine/Source/Runtime/D3D12RHI/Private/D3D12Commands.cpp
@@ -125,9 +125,6 @@ FD3D12ResourceLocation& FD3D12CommandContext::SetupIndirectArgument(FRHIBuffer*
 {
 	FD3D12ResourceLocation& ArgumentBufferLocation = RetrieveObject<FD3D12Buffer>(ArgumentBufferRHI)->ResourceLocation;
 
-	// Indirect args buffer can be a previously pending UAV, which becomes PS\Non-PS read. ApplyState will flush pending transitions, so enqueue the indirect arg transition and flush here.
-	TransitionResource(ArgumentBufferLocation.GetResource(), D3D12_RESOURCE_STATE_TBD, D3D12_RESOURCE_STATE_INDIRECT_ARGUMENT | ExtraStates, D3D12_RESOURCE_BARRIER_ALL_SUBRESOURCES);
-
 	// Must flush so the desired state is actually set.
 	FlushResourceBarriers();
 
@@ -241,9 +238,7 @@ void ProcessResource(FD3D12CommandContext& Context, const FRHITransitionInfo& In
 		FD3D12UnorderedAccessView* UAV = Context.RetrieveObject<FD3D12UnorderedAccessView_RHI>(Info.UAV);
 		check(UAV);
 
-		FRHITransitionInfo LocalInfo = Info;
-		LocalInfo.MipIndex = UAV->GetViewSubset().Range.MostDetailedMip();
-		Function(LocalInfo, UAV->GetResource());
+		Function(Info, UAV->GetResource(), nullptr, nullptr);
 		break;
 	}
 	case FRHITransitionInfo::EType::Buffer:
@@ -253,7 +248,7 @@ void ProcessResource(FD3D12CommandContext& Context, const FRHITransitionInfo& In
 		check(Buffer || GNumExplicitGPUsForRendering > 1);
 		if (Buffer)
 		{
-			Function(Info, Buffer->GetResource());
+			Function(Info, Buffer->GetResource(), nullptr, nullptr);
 		}
 		break;
 	}
@@ -264,14 +259,14 @@ void ProcessResource(FD3D12CommandContext& Context, const FRHITransitionInfo& In
 		check(Texture || GNumExplicitGPUsForRendering > 1);
 		if (Texture)
 		{
-			FD3D12Texture* TextureOut = nullptr;
+			FD3D12Texture* DiscardTextureOut = nullptr;
 #if PLATFORM_REQUIRES_TYPELESS_RESOURCE_DISCARD_WORKAROUND
 			if (Texture->GetRequiresTypelessResourceDiscardWorkaround())
 			{
-				TextureOut = Texture;
+				DiscardTextureOut = Texture;
 			}
 #endif // #if PLATFORM_REQUIRES_TYPELESS_RESOURCE_DISCARD_WORKAROUND
-			Function(Info, Texture->GetResource(), TextureOut);
+			Function(Info, Texture->GetResource(), Texture, DiscardTextureOut);
 		}
 		break;
 	}
@@ -281,10 +276,44 @@ void ProcessResource(FD3D12CommandContext& Context, const FRHITransitionInfo& In
 	}
 }
 
+// static bool ProcessTransitionDuringBegin(const FD3D12TransitionData* Data)
+// {
+// 	// Pipe changes which are not ending with graphics or targeting all pipelines are handle during begin
+// 	return !EnumHasAllFlags(Data->SrcPipelines, ERHIPipeline::All) && ((Data->SrcPipelines != Data->DstPipelines && Data->DstPipelines != ERHIPipeline::Graphics) || EnumHasAllFlags(Data->DstPipelines, ERHIPipeline::All));
+// }
+
+// Pipe changes which are not ending with graphics or targeting all pipelines are handle during begin
 static bool ProcessTransitionDuringBegin(const FD3D12TransitionData* Data)
 {
-	// Pipe changes which are not ending with graphics or targeting all pipelines are handle during begin
-	return !EnumHasAllFlags(Data->SrcPipelines, ERHIPipeline::All) && ((Data->SrcPipelines != Data->DstPipelines && Data->DstPipelines != ERHIPipeline::Graphics) || EnumHasAllFlags(Data->DstPipelines, ERHIPipeline::All));
+	// Source pipelines aren't on all pipelines
+	const bool bSrcPipelinesNotAll = !EnumHasAllFlags(Data->SrcPipelines, ERHIPipeline::All);
+
+	// Source and destination pipelines are different
+	const bool bSrcDstPipelinesDiffer = Data->SrcPipelines != Data->DstPipelines;
+
+	// Destination pipeline is not only graphics
+	const bool bDstPipelineNotGraphics = Data->DstPipelines != ERHIPipeline::Graphics;
+
+	// Destination pipelines include all pipelines
+	const bool bDstPipelinesIncludeAll = EnumHasAllFlags(Data->DstPipelines, ERHIPipeline::All);
+
+	return bSrcPipelinesNotAll && ( (bSrcDstPipelinesDiffer && bDstPipelineNotGraphics) || bDstPipelinesIncludeAll );
+}
+
+static bool ShouldProcessTransition(const FD3D12TransitionData* Data, bool bIsBeginTransition, ERHIPipeline ExecutingPipeline)
+{
+	// Special DX12 case where crosspipe transitions from AsyncCompute with graphics state can only be processed on the ERHIPipeline::Graphics pipe
+	if (Data->bAsyncToAllPipelines)
+	{
+		if ((bIsBeginTransition == false) && (ExecutingPipeline == ERHIPipeline::Graphics))
+		{
+			return true;
+		}
+
+		return false;
+	}
+
+	return (bIsBeginTransition ? ProcessTransitionDuringBegin(Data) : !ProcessTransitionDuringBegin(Data));
 }
 
 struct FD3D12DiscardResource
@@ -338,7 +367,7 @@ void FD3D12CommandContext::HandleResourceDiscardTransitions(
 			continue;
 		}
 
-		ProcessResource(*this, Info, [&](const FRHITransitionInfo& Info, FD3D12Resource* Resource, FD3D12Texture* Texture = nullptr)
+		ProcessResource(*this, Info, [&](const FRHITransitionInfo& Info, FD3D12Resource* Resource, FD3D12Texture* Texture, FD3D12Texture* DiscardTexture)
 		{
 			if (!Resource->RequiresResourceStateTracking())
 			{
@@ -348,23 +377,41 @@ void FD3D12CommandContext::HandleResourceDiscardTransitions(
 			D3D12_RESOURCE_STATES InitialState = GetInitialResourceState(Resource->GetDesc(), QueueType);
 			if (Info.IsWholeResource() || Resource->GetSubresourceCount() == 1)
 			{
-				TransitionResource(Resource, D3D12_RESOURCE_STATE_TBD, InitialState, D3D12_RESOURCE_BARRIER_ALL_SUBRESOURCES);
+				if ((Info.DiscardAccessBefore != ERHIAccess::Unknown) && (Info.DiscardAccessBefore != ERHIAccess::Discard))
+				{
+					D3D12_RESOURCE_STATES StateBefore = GetD3D12ResourceState(Info.DiscardAccessBefore, QueueType, Texture);
+					TransitionResource(Resource, StateBefore, InitialState, D3D12_RESOURCE_BARRIER_ALL_SUBRESOURCES);
+				}
+				else
+				{
+					UpdateResidency(Resource);
+				}
+
 
 				const FD3D12RenderTargetView* RTV = nullptr;
 #if PLATFORM_REQUIRES_TYPELESS_RESOURCE_DISCARD_WORKAROUND
-				if (Texture)
+				if (DiscardTexture)
 				{
-					RTV = Texture->GetRenderTargetView(0, -1);
+					RTV = DiscardTexture->GetRenderTargetView(0, -1);
 				}
 #endif // #if PLATFORM_REQUIRES_TYPELESS_RESOURCE_DISCARD_WORKAROUND				
 				ResourcesToDiscard.Emplace(Resource, Info.Flags, D3D12_RESOURCE_BARRIER_ALL_SUBRESOURCES, Texture, RTV);
 			}
 			else
 			{
-				EnumerateSubresources(Resource, Info, Texture, [&](uint32 Subresource, const FD3D12RenderTargetView* RTV)
+				EnumerateSubresources(Resource, Info, DiscardTexture, [&](uint32 Subresource, const FD3D12RenderTargetView* RTV)
 				{
-					TransitionResource(Resource, D3D12_RESOURCE_STATE_TBD, InitialState, Subresource);
-					ResourcesToDiscard.Emplace(Resource, Info.Flags, Subresource, Texture, RTV);
+					if ((Info.DiscardAccessBefore != ERHIAccess::Unknown) && (Info.DiscardAccessBefore != ERHIAccess::Discard))
+					{
+						D3D12_RESOURCE_STATES StateBefore = GetD3D12ResourceState(Info.DiscardAccessBefore, QueueType, Texture);
+						TransitionResource(Resource, StateBefore, InitialState, Subresource);
+					}
+					else
+					{
+						UpdateResidency(Resource);
+					}
+
+					ResourcesToDiscard.Emplace(Resource, Info.Flags, Subresource, DiscardTexture, RTV);
 				});
 			}
 		});
@@ -492,59 +539,77 @@ void FD3D12CommandContext::HandleResourceTransitions(const FD3D12TransitionData*
 			bUAVBarrier |= bUAVAccessAfter;
 		}
 
-		// Only need to check for UAV barriers here, all other transitions are handled inside the RHI
-		// Can't rely on before state to either be unknown or UAV because there could be a UAV->SRV transition already done (and ignored here)
-		// and the transition SRV->UAV needs a UAV barrier then to work correctly otherwise there is no synchronization at all
+		// The transition SRV->UAV needs a UAV barrier then to work correctly otherwise there is no synchronization at all
 		bUAVBarrier |= bUAVAccessAfter;
 
-		const bool bHandleDiscardTransition = (Info.AccessAfter == ERHIAccess::Discard && QueueType == ED3D12QueueType::Direct);
+
+
+		const bool bDiscardTransitionOnAsync = (Info.AccessAfter == ERHIAccess::Discard && QueueType == ED3D12QueueType::Async);
 
 		// Process transitions which are forced during begin because those contain transition from Graphics to Compute and should
 		// help remove forced patch up command lists for async compute to run on the graphics queue
-		if (Info.Resource && (ProcessTransitionDuringBegin(TransitionData) || bHandleDiscardTransition))
+		if (Info.Resource && (!bDiscardTransitionOnAsync))
 		{
-			ProcessResource(*this, Info, [&](const FRHITransitionInfo& Info, FD3D12Resource* Resource, FD3D12Texture* UnusedTexture = nullptr)
-			{
-				if (!Resource->RequiresResourceStateTracking())
+			ProcessResource(*this, Info, [&](const FRHITransitionInfo& Info, FD3D12Resource* Resource, FD3D12Texture* Texture, FD3D12Texture* DiscardTexture)
 				{
-					return;
-				}
+					if (!Resource->RequiresResourceStateTracking())
+					{
+						return;
+					}
 
-				// Use D3D12_RESOURCE_STATE_TBD as before state for now and don't use provided before state because there is not validation nor handling if the provided state
-				// doesn't match up with the tracked state - this needs to be improved and checked
-				D3D12_RESOURCE_STATES StateBefore = D3D12_RESOURCE_STATE_TBD;
-				D3D12_RESOURCE_STATES StateAfter = Info.AccessAfter == ERHIAccess::Discard ? GetInitialResourceState(Resource->GetDesc(), QueueType) : GetD3D12ResourceState(Info.AccessAfter, QueueType);
+					check(Info.AccessBefore != ERHIAccess::Unknown);
+					check(Info.AccessAfter != ERHIAccess::Unknown);
 
-				// enqueue the correct transitions
-				if (Info.IsWholeResource() || Resource->GetSubresourceCount() == 1)
-				{
-					TransitionResource(Resource, StateBefore, StateAfter, D3D12_RESOURCE_BARRIER_ALL_SUBRESOURCES);
-				}
-				else
-				{
-					EnumerateSubresources(Resource, Info, nullptr, [&](uint32 Subresource, const FD3D12RenderTargetView* UnusedRTV = nullptr)
+					// special case that can happen during first frames or resize
+					if ((Info.AccessBefore == ERHIAccess::Present) && (Info.AccessAfter == ERHIAccess::Present))
+						return;
+						
+					D3D12_RESOURCE_STATES StateBefore = Info.AccessBefore == ERHIAccess::Discard ? GetInitialResourceState(Resource->GetDesc(), QueueType) : GetD3D12ResourceState(Info.AccessBefore, QueueType, Texture);
+						
+					if (Info.AccessBefore != ERHIAccess::Present)
+					{
+						check(StateBefore != D3D12_RESOURCE_STATE_COMMON);
+					}
+
+					D3D12_RESOURCE_STATES StateAfter = Info.AccessAfter == ERHIAccess::Discard ? GetInitialResourceState(Resource->GetDesc(), QueueType) : GetD3D12ResourceState(Info.AccessAfter, QueueType, Texture);
+
+					// enqueue the correct transitions
+					if (Info.IsWholeResource() || Resource->GetSubresourceCount() == 1)
+					{
+						TransitionResource(Resource, StateBefore, StateAfter, D3D12_RESOURCE_BARRIER_ALL_SUBRESOURCES);
+					}
+					else
 					{
+						// high level rendering is controlling transition ranges, at this level this is an index not a range
+						check(Info.MipIndex != FRHISubresourceRange::kAllSubresources);
+						check(Info.ArraySlice != FRHISubresourceRange::kAllSubresources);
+						check(Info.PlaneSlice != FRHISubresourceRange::kAllSubresources);
+						const uint32 Subresource = D3D12CalcSubresource(Info.MipIndex, Info.ArraySlice, Info.PlaneSlice, Resource->GetMipLevels(), Resource->GetArraySize());
+						check(Subresource < Resource->GetSubresourceCount());
 						TransitionResource(Resource, StateBefore, StateAfter, Subresource);
-					});
-				}
-			});
+					}
+				});
 		}
 	}
+	
 }
 
 void FD3D12CommandContext::RHIBeginTransitions(TArrayView<const FRHITransition*> Transitions)
 {
+	const ERHIPipeline CurrentPipeline = GetPipeline();
+	const bool bIsBeginTransition = true;
+
 	for (const FRHITransition* Transition : Transitions)
 	{
 		const FD3D12TransitionData* Data = Transition->GetPrivateData<FD3D12TransitionData>();
 
-		if (ProcessTransitionDuringBegin(Data))
+		if (ShouldProcessTransition ( Data, bIsBeginTransition, CurrentPipeline))
 		{
 			HandleTransientAliasing(Data);
 		}
 	}
 
-	HandleDiscardResources(Transitions, true /** bIsBeginTransitions */);
+	HandleDiscardResources(Transitions, bIsBeginTransition );
 
 	bool bUAVBarrier = false;
 
@@ -553,7 +618,7 @@ void FD3D12CommandContext::RHIBeginTransitions(TArrayView<const FRHITransition*>
 		const FD3D12TransitionData* Data = Transition->GetPrivateData<FD3D12TransitionData>();
 
 		// Handle transition during BeginTransitions?
-		if (ProcessTransitionDuringBegin(Data))
+		if (ShouldProcessTransition(Data, bIsBeginTransition, CurrentPipeline))
 		{
 			HandleResourceTransitions(Data, bUAVBarrier);
 		}
@@ -565,7 +630,6 @@ void FD3D12CommandContext::RHIBeginTransitions(TArrayView<const FRHITransition*>
 	}
 
 	// Signal fences
-	const ERHIPipeline SourcePipeline = GetPipeline();
 	for (const FRHITransition* Transition : Transitions)
 	{
 		const auto* Data = Transition->GetPrivateData<FD3D12TransitionData>();
@@ -573,9 +637,9 @@ void FD3D12CommandContext::RHIBeginTransitions(TArrayView<const FRHITransition*>
 		{
 			const TRHIPipelineArray<FD3D12SyncPointRef>& DeviceSyncPoints = Data->SyncPoints[GetGPUIndex()];
 
-			if (DeviceSyncPoints[SourcePipeline])
+			if (DeviceSyncPoints[CurrentPipeline])
 			{
-				SignalSyncPoint(DeviceSyncPoints[SourcePipeline]);
+				SignalSyncPoint(DeviceSyncPoints[CurrentPipeline]);
 			}
 		}
 	}
@@ -583,9 +647,11 @@ void FD3D12CommandContext::RHIBeginTransitions(TArrayView<const FRHITransition*>
 
 void FD3D12CommandContext::RHIEndTransitions(TArrayView<const FRHITransition*> Transitions)
 {
+	const ERHIPipeline CurrentPipeline = GetPipeline();
+	const bool bIsBeginTransition = false;
+
 	// Wait for fences
 	{
-		const ERHIPipeline DstPipeline = GetPipeline();
 		for (const FRHITransition* Transition : Transitions)
 		{
 			const auto* Data = Transition->GetPrivateData<FD3D12TransitionData>();
@@ -595,7 +661,7 @@ void FD3D12CommandContext::RHIEndTransitions(TArrayView<const FRHITransition*> T
 
 				for (ERHIPipeline SrcPipeline : MakeFlagsRange(Data->SrcPipelines))
 				{
-					if (SrcPipeline != DstPipeline && DeviceSyncPoints[SrcPipeline])
+					if (SrcPipeline != CurrentPipeline && DeviceSyncPoints[SrcPipeline])
 					{
 						WaitSyncPoint(DeviceSyncPoints[SrcPipeline]);
 					}
@@ -608,6 +674,7 @@ void FD3D12CommandContext::RHIEndTransitions(TArrayView<const FRHITransition*> T
 	for (const FRHITransition* Transition : Transitions)
 	{
 		const FD3D12TransitionData* Data = Transition->GetPrivateData<FD3D12TransitionData>();
+
 		HandleReservedResourceCommits(Data);
 	}
 
@@ -615,7 +682,8 @@ void FD3D12CommandContext::RHIEndTransitions(TArrayView<const FRHITransition*> T
 	{
 		const FD3D12TransitionData* Data = Transition->GetPrivateData<FD3D12TransitionData>();
 
-		if (!ProcessTransitionDuringBegin(Data))
+		//if (!ProcessTransitionDuringBegin(Data))
+		if (ShouldProcessTransition(Data, bIsBeginTransition, CurrentPipeline))
 		{
 			HandleTransientAliasing(Data);
 		}
@@ -630,7 +698,8 @@ void FD3D12CommandContext::RHIEndTransitions(TArrayView<const FRHITransition*> T
 		const FD3D12TransitionData* Data = Transition->GetPrivateData<FD3D12TransitionData>();
 
 		// Handle transition during EndTransitions?
-		if (!ProcessTransitionDuringBegin(Data))
+		//if (!ProcessTransitionDuringBegin(Data))
+		if (ShouldProcessTransition(Data, bIsBeginTransition, CurrentPipeline))
 		{
 			HandleResourceTransitions(Data, bUAVBarrier);
 		}
@@ -709,14 +778,24 @@ void FD3D12CommandContext::RHICopyToStagingBuffer(FRHIBuffer* SourceBufferRHI, F
 		D3D12_RESOURCE_DESC const& DestBufferDesc = pDestResource->GetDesc();
 		uint32 DestOffset = StagingBuffer->ResourceLocation.GetOffsetFromBaseOfResource();
 
-		if (pSourceResource->RequiresResourceStateTracking())
+		FlushResourceBarriers();	// Must flush so the desired state is actually set.
+
+#if D3D12_RHI_RAYTRACING
+		// TODO: Buffer in D3D12_RESOURCE_STATE_RAYTRACING_ACCELERATION_STRUCTURE needs CopyRaytracingAccelerationStructure
+		bool bRayTracingAccellerationStruct = ((pSourceResource->RequiresResourceStateTracking() == false) && (pSourceResource->GetDefaultResourceState() == D3D12_RESOURCE_STATE_RAYTRACING_ACCELERATION_STRUCTURE));
+#else
+		bool bRayTracingAccellerationStruct = false;
+#endif
+
+		if (bRayTracingAccellerationStruct)
 		{
-			TransitionResource(pSourceResource, D3D12_RESOURCE_STATE_TBD, D3D12_RESOURCE_STATE_COPY_SOURCE, 0);
+			UE_LOG(LogD3D12RHI, Warning, TEXT("RHICopyToStagingBuffer cannot be used on the RayTracing Accelleration structure %s"), *SourceBufferRHI->GetName().GetPlainNameString());
+		}
+		else
+		{
+			GraphicsCommandList()->CopyBufferRegion(pDestResource->GetResource(), DestOffset, pSourceResource->GetResource(), Offset + SourceOffset, NumBytes);
 		}
 		
-		FlushResourceBarriers();	// Must flush so the desired state is actually set.
-
-		GraphicsCommandList()->CopyBufferRegion(pDestResource->GetResource(), DestOffset, pSourceResource->GetResource(), Offset + SourceOffset, NumBytes);
 		UpdateResidency(pDestResource);
 		UpdateResidency(pSourceResource);
 
@@ -1409,13 +1488,6 @@ void FD3D12CommandContext::SetRenderTargetsAndClear(const FRHISetRenderTargetsIn
 		if (RenderTargetsInfo.ShadingRateTexture != nullptr)
 		{
 			FD3D12Resource* Resource = RetrieveTexture(RenderTargetsInfo.ShadingRateTexture)->GetResource();
-
-			TransitionResource(
-				Resource,
-				D3D12_RESOURCE_STATE_TBD,
-				D3D12_RESOURCE_STATE_SHADING_RATE_SOURCE,
-				D3D12_RESOURCE_BARRIER_ALL_SUBRESOURCES);
-
 			StateCache.SetShadingRateImage(Resource);
 		}
 		else
@@ -1740,7 +1812,7 @@ void FD3D12CommandContext::RHIMultiDrawIndexedPrimitiveIndirect(FRHIBuffer* Inde
 		FD3D12Buffer* CountBuffer = RetrieveObject<FD3D12Buffer>(CountBufferRHI);
 		FD3D12ResourceLocation& CounterLocation = CountBuffer->ResourceLocation;
 		CountBufferResource = CounterLocation.GetResource()->GetResource();
-		TransitionResource(CounterLocation.GetResource(), D3D12_RESOURCE_STATE_TBD, D3D12_RESOURCE_STATE_INDIRECT_ARGUMENT, D3D12_RESOURCE_BARRIER_ALL_SUBRESOURCES);
+		
 		CountBufferOffsetFromResourceBase = CounterLocation.GetOffsetFromBaseOfResource() + CountBufferOffset;
 		UpdateResidency(CounterLocation.GetResource());
 	}
diff --git a/Engine/Source/Runtime/D3D12RHI/Private/D3D12DescriptorCache.cpp b/Engine/Source/Runtime/D3D12RHI/Private/D3D12DescriptorCache.cpp
--- a/Engine/Source/Runtime/D3D12RHI/Private/D3D12DescriptorCache.cpp
+++ b/Engine/Source/Runtime/D3D12RHI/Private/D3D12DescriptorCache.cpp
@@ -215,18 +215,6 @@ void FD3D12DescriptorCache::SetVertexBuffers(FD3D12VertexBufferCache& Cache)
 		if (Cache.CurrentVertexBufferResources[i])
 		{
 			Context.UpdateResidency(Cache.Resources[i]);
-
-			FD3D12Resource* Resource = Cache.CurrentVertexBufferResources[i]->GetResource();
-			if (Resource && Resource->RequiresResourceStateTracking())
-			{
-				check(Resource->GetSubresourceCount() == 1);
-				Context.TransitionResource(
-					Resource,
-					D3D12_RESOURCE_STATE_TBD,
-					D3D12_RESOURCE_STATE_VERTEX_AND_CONSTANT_BUFFER,
-					D3D12_RESOURCE_BARRIER_ALL_SUBRESOURCES
-				);
-			}
 		}
 	}
 }
@@ -261,7 +249,6 @@ D3D12_GPU_DESCRIPTOR_HANDLE FD3D12DescriptorCache::BuildUAVTable(EShaderFrequenc
 		{
 			SrcDescriptors[SlotIndex] = UAVs[SlotIndex]->GetOfflineCpuHandle();
 
-			Context.TransitionResource(UAVs[SlotIndex], D3D12_RESOURCE_STATE_UNORDERED_ACCESS);
 			Context.UpdateResidency(Cache.Resources[ShaderStage][SlotIndex]);
 		}
 	}
@@ -320,8 +307,6 @@ void FD3D12DescriptorCache::SetRenderTargets(FD3D12RenderTargetView** RenderTarg
 	{
 		if (RenderTargetViewArray[i] != NULL)
 		{
-			Context.TransitionResource(RenderTargetViewArray[i], D3D12_RESOURCE_STATE_RENDER_TARGET);
-
 			RTVDescriptors[i] = RenderTargetViewArray[i]->GetOfflineCpuHandle();
 			Context.UpdateResidency(RenderTargetViewArray[i]->GetResource());
 		}
@@ -333,8 +318,6 @@ void FD3D12DescriptorCache::SetRenderTargets(FD3D12RenderTargetView** RenderTarg
 
 	if (DepthStencilTarget != nullptr)
 	{
-		Context.TransitionResource(DepthStencilTarget);
-
 		const D3D12_CPU_DESCRIPTOR_HANDLE DSVDescriptor = DepthStencilTarget->GetOfflineCpuHandle();
 		Context.GraphicsCommandList()->OMSetRenderTargets(Count, RTVDescriptors, 0, &DSVDescriptor);
 		Context.UpdateResidency(DepthStencilTarget->GetResource());
@@ -455,24 +438,6 @@ void FD3D12DescriptorCache::SetSamplerTable(EShaderFrequency ShaderStage, const
 #endif
 }
 
-static D3D12_RESOURCE_STATES GetBindingResourceState(FD3D12CommandContext& Context, EShaderFrequency ShaderStage, FD3D12ShaderResourceView* SRV)
-{
-	D3D12_RESOURCE_STATES State = (ShaderStage == SF_Compute)
-		? D3D12_RESOURCE_STATE_NON_PIXEL_SHADER_RESOURCE
-		: D3D12_RESOURCE_STATE_NON_PIXEL_SHADER_RESOURCE | D3D12_RESOURCE_STATE_PIXEL_SHADER_RESOURCE;
-
-	if (SRV->GetResource()->IsDepthStencilResource())
-	{
-		State |= D3D12_RESOURCE_STATE_DEPTH_READ;
-	}
-	else if (SRV->GetSkipFastClearFinalize())
-	{
-		State |= Context.SkipFastClearEliminateState;
-	}
-
-	return State;
-}
-
 D3D12_GPU_DESCRIPTOR_HANDLE FD3D12DescriptorCache::BuildSRVTable(EShaderFrequency ShaderStage, const FD3D12RootSignature* RootSignature, FD3D12ShaderResourceViewCache& Cache, const SRVSlotMask& SlotsNeededMask, uint32 SlotsNeeded, uint32& HeapSlot)
 {
 	SRVSlotMask& CurrentDirtySlotMask = Cache.DirtySlotMask[ShaderStage];
@@ -497,9 +462,6 @@ D3D12_GPU_DESCRIPTOR_HANDLE FD3D12DescriptorCache::BuildSRVTable(EShaderFrequenc
 		{
 			SrcDescriptors[SlotIndex] = SRV->GetOfflineCpuHandle();
 
-			const D3D12_RESOURCE_STATES State = GetBindingResourceState(Context, ShaderStage, SRV);
-
-			Context.TransitionResource(SRV, State & ValidResourceStates);
 			Context.UpdateResidency(Cache.Resources[ShaderStage][SlotIndex]);
 
 #if RHI_RAYTRACING
@@ -560,8 +522,6 @@ void FD3D12DescriptorCache::PrepareBindlessViews(EShaderFrequency ShaderStage, T
 	{
 		if (ensure(SRV))
 		{
-			const D3D12_RESOURCE_STATES State = GetBindingResourceState(Context, ShaderStage, SRV);
-			Context.TransitionResource(SRV, State & ValidResourceStates);
 			Context.UpdateResidency(SRV->GetResource());
 		}
 	}
@@ -570,7 +530,6 @@ void FD3D12DescriptorCache::PrepareBindlessViews(EShaderFrequency ShaderStage, T
 	{
 		if (ensure(UAV))
 		{
-			Context.TransitionResource(UAV, D3D12_RESOURCE_STATE_UNORDERED_ACCESS);
 			Context.UpdateResidency(UAV->GetResource());
 		}
 	}
diff --git a/Engine/Source/Runtime/D3D12RHI/Private/D3D12Device.cpp b/Engine/Source/Runtime/D3D12RHI/Private/D3D12Device.cpp
--- a/Engine/Source/Runtime/D3D12RHI/Private/D3D12Device.cpp
+++ b/Engine/Source/Runtime/D3D12RHI/Private/D3D12Device.cpp
@@ -406,7 +406,7 @@ void FD3D12Device::SetupAfterDeviceCreation()
 			DispatchRaysDescBufferDesc, 256,
 			FRHIBufferDesc(DispatchRaysDescBufferDesc.Width, 0, BUF_DrawIndirect),
 			ED3D12ResourceStateMode::MultiState, D3D12_RESOURCE_STATE_INDIRECT_ARGUMENT, false /*bInitialData*/,
-			GetGPUMask(), nullptr /*ResourceAllocator*/, TEXT("DispatchRaysDescBuffer"));
+			GetGPUMask(), nullptr /*ResourceAllocator*/, TEXT("DispatchRaysDescBuffer"), ERHIAccess::IndirectArgs);
 	}
 #endif // D3D12_RHI_RAYTRACING
 
diff --git a/Engine/Source/Runtime/D3D12RHI/Private/D3D12PoolAllocator.cpp b/Engine/Source/Runtime/D3D12RHI/Private/D3D12PoolAllocator.cpp
--- a/Engine/Source/Runtime/D3D12RHI/Private/D3D12PoolAllocator.cpp
+++ b/Engine/Source/Runtime/D3D12RHI/Private/D3D12PoolAllocator.cpp
@@ -111,7 +111,7 @@ void FD3D12MemoryPool::Init()
 			LLM_SCOPE_BYTAG(D3D12AllocatorUnused);
 #endif
 			const D3D12_HEAP_PROPERTIES HeapProps = CD3DX12_HEAP_PROPERTIES(InitConfig.HeapType, GetGPUMask().GetNative(), GetVisibilityMask().GetNative());
-			VERIFYD3D12RESULT(Adapter->CreateBuffer(HeapProps, GetGPUMask(), InitConfig.InitialResourceState, ED3D12ResourceStateMode::SingleState, InitConfig.InitialResourceState, PoolSize, BackingResource.GetInitReference(), TEXT("Resource Allocator Underlying Buffer"), InitConfig.ResourceFlags));
+			VERIFYD3D12RESULT(Adapter->CreateBuffer(HeapProps, GetGPUMask(), InitConfig.InitialResourceState, ED3D12ResourceStateMode::SingleState, InitConfig.InitialResourceState, PoolSize, BackingResource.GetInitReference(), TEXT("Resource PoolAllocator Underlying Buffer"), InitConfig.ResourceFlags));
 #if UE_MEMORY_TRACE_ENABLED
 			MemoryTrace_MarkAllocAsHeap(BackingResource->GetGPUVirtualAddress(), TraceHeapId);
 #endif
@@ -590,6 +590,40 @@ bool FD3D12PoolAllocator::HandleDefragRequest(FRHIContextArray const& Contexts,
 	uint64 CurrentOffset = Owner->GetOffsetFromBaseOfResource();
 	FD3D12Resource* CurrentResource = Owner->GetResource();
 
+	D3D12_RESOURCE_STATES SourceResourceState = D3D12_RESOURCE_STATE_TBD;
+	if (CurrentResource->RequiresResourceStateTracking())
+	{
+
+		// State tracking is happening at the RHI high level 
+		ERHIAccess SourceResourceAccess = ERHIAccess::Unknown;
+		FD3D12Texture* CurrentResourceRHITexture = nullptr;
+		if (Owner->Owner)
+		{
+			if (CurrentResource->GetDesc().Dimension == D3D12_RESOURCE_DIMENSION_BUFFER)
+			{
+				FD3D12Buffer* CurrentResourceRHIBuffer = (FD3D12Buffer*)Owner->Owner;
+				SourceResourceAccess = CurrentResourceRHIBuffer->GetTrackedAccess_Unsafe();
+			}
+			else
+			{
+				CurrentResourceRHITexture = (FD3D12Texture*)Owner->Owner;
+				SourceResourceAccess = CurrentResourceRHITexture->GetTrackedAccess_Unsafe();
+			}
+		}
+
+		SourceResourceState = SourceResourceAccess == ERHIAccess::Unknown ? CurrentResource->GetInitialResourceState() : GetD3D12ResourceState(SourceResourceAccess, ED3D12QueueType::Direct, CurrentResourceRHITexture);
+
+		if (SourceResourceState == D3D12_RESOURCE_STATE_TBD)
+		{
+			SourceResourceState = CurrentResource->GetReadableState();
+		}
+	}
+	else
+	{
+		//If the resource is a single state type and doesn't need state tracking set the state as TBD so the scoped transition in FlushPendingCopyOps can revert back to the default state
+		SourceResourceState = D3D12_RESOURCE_STATE_TBD;
+	}
+
 	// Release the current allocation (will only be freed on the next frame fence)
 	bool bDefragFree = true;
 	DeallocateResource(*Owner, bDefragFree);
@@ -601,7 +635,21 @@ bool FD3D12PoolAllocator::HandleDefragRequest(FRHIContextArray const& Contexts,
 	Owner->SetPoolAllocator(this);
 
 	// Notify owner of moved allocation data (recreated resources and SRVs if needed)
-	Owner->OnAllocationMoved(Contexts, InSourceBlock);
+	D3D12_RESOURCE_STATES DestCreateState = D3D12_RESOURCE_STATE_TBD;
+	Owner->OnAllocationMoved(Contexts, InSourceBlock, DestCreateState);
+
+	if (Owner->GetResource()->RequiresResourceStateTracking())
+	{
+		if (DestCreateState == D3D12_RESOURCE_STATE_TBD)
+		{
+			DestCreateState = D3D12_RESOURCE_STATE_COMMON;
+		}
+	}
+	else
+	{
+		//If the resource is a single state type and doesn't need state tracking set the state as TBD so the scoped transition in FlushPendingCopyOps can revert back to the default state
+		DestCreateState = D3D12_RESOURCE_STATE_TBD;
+	}
 
 	// Add request to unlock the source block on the next fence value (copy operation should have been done by then)
 	FD3D12Adapter* Adapter = GetParentDevice()->GetParentAdapter();
@@ -615,8 +663,10 @@ bool FD3D12PoolAllocator::HandleDefragRequest(FRHIContextArray const& Contexts,
 	// Schedule a copy operation of the actual data
 	FD3D12VRAMCopyOperation CopyOp;
 	CopyOp.SourceResource	= CurrentResource;
+	CopyOp.SourceResourceState = SourceResourceState;
 	CopyOp.SourceOffset		= CurrentOffset;
 	CopyOp.DestResource		= Owner->GetResource();
+	CopyOp.DestResourceState = DestCreateState;
 	CopyOp.DestOffset		= Owner->GetOffsetFromBaseOfResource();
 	CopyOp.Size				= InSourceBlock->GetSize();
 	CopyOp.CopyType			= AllocationStrategy == EResourceAllocationStrategy::kManualSubAllocation ? FD3D12VRAMCopyOperation::ECopyType::BufferRegion : FD3D12VRAMCopyOperation::ECopyType::Resource;
@@ -881,32 +931,47 @@ void FD3D12PoolAllocator::FlushPendingCopyOps(FD3D12CommandContext& InCommandCon
 		else
 #endif // D3D12_RHI_RAYTRACING
 		{
-			FScopedResourceBarrier ScopedResourceBarrierSrc(InCommandContext, CopyOperation.SourceResource, nullptr, D3D12_RESOURCE_STATE_COPY_SOURCE, D3D12_RESOURCE_BARRIER_ALL_SUBRESOURCES);
-			FScopedResourceBarrier ScopedResourceBarrierDst(InCommandContext, CopyOperation.DestResource  , nullptr, D3D12_RESOURCE_STATE_COPY_DEST  , D3D12_RESOURCE_BARRIER_ALL_SUBRESOURCES);
-
-			InCommandContext.FlushResourceBarriers();
 
-			switch (CopyOperation.CopyType)
-			{
-			case FD3D12VRAMCopyOperation::BufferRegion:
 			{
-				InCommandContext.GraphicsCommandList()->CopyBufferRegion(
-					CopyOperation.DestResource->GetResource(),
-					CopyOperation.DestOffset,
-					CopyOperation.SourceResource->GetResource(),
-					CopyOperation.SourceOffset,
-					CopyOperation.Size
-				);
-				break;
+				FScopedResourceBarrier ScopedResourceBarrierSrc(InCommandContext, CopyOperation.SourceResource, nullptr, D3D12_RESOURCE_STATE_COPY_SOURCE, D3D12_RESOURCE_BARRIER_ALL_SUBRESOURCES, CopyOperation.SourceResourceState);
+				FScopedResourceBarrier ScopedResourceBarrierDst(InCommandContext, CopyOperation.DestResource, nullptr, D3D12_RESOURCE_STATE_COPY_DEST, D3D12_RESOURCE_BARRIER_ALL_SUBRESOURCES, CopyOperation.DestResourceState);
+
+				InCommandContext.FlushResourceBarriers();
+
+				switch (CopyOperation.CopyType)
+				{
+				case FD3D12VRAMCopyOperation::BufferRegion:
+				{
+					InCommandContext.GraphicsCommandList()->CopyBufferRegion(
+						CopyOperation.DestResource->GetResource(),
+						CopyOperation.DestOffset,
+						CopyOperation.SourceResource->GetResource(),
+						CopyOperation.SourceOffset,
+						CopyOperation.Size
+					);
+					break;
+				}
+				case FD3D12VRAMCopyOperation::Resource:
+				{
+					InCommandContext.GraphicsCommandList()->CopyResource(
+						CopyOperation.DestResource->GetResource(),
+						CopyOperation.SourceResource->GetResource()
+					);
+					break;
+				}
+				}
 			}
-			case FD3D12VRAMCopyOperation::Resource:
+
+			// if the resource need tracking we are setting the state of the destination equal to the original source state
+			if (CopyOperation.SourceResource->RequiresResourceStateTracking())
 			{
-				InCommandContext.GraphicsCommandList()->CopyResource(
-					CopyOperation.DestResource->GetResource(),
-					CopyOperation.SourceResource->GetResource()
-				);
-				break;
-			}
+				check(CopyOperation.DestResource->RequiresResourceStateTracking());
+				check(CopyOperation.SourceResourceState != D3D12_RESOURCE_STATE_TBD);
+
+				if (CopyOperation.SourceResourceState != CopyOperation.DestResourceState)
+				{
+					InCommandContext.AddTransitionBarrier(CopyOperation.DestResource, CopyOperation.DestResourceState, CopyOperation.SourceResourceState, D3D12_RESOURCE_BARRIER_ALL_SUBRESOURCES);
+				}
 			}
 		}
 
diff --git a/Engine/Source/Runtime/D3D12RHI/Private/D3D12PoolAllocator.h b/Engine/Source/Runtime/D3D12RHI/Private/D3D12PoolAllocator.h
--- a/Engine/Source/Runtime/D3D12RHI/Private/D3D12PoolAllocator.h
+++ b/Engine/Source/Runtime/D3D12RHI/Private/D3D12PoolAllocator.h
@@ -59,8 +59,10 @@ struct FD3D12VRAMCopyOperation
 	};
 
 	FD3D12Resource* SourceResource;
+	D3D12_RESOURCE_STATES SourceResourceState;
 	uint32 SourceOffset;
 	FD3D12Resource* DestResource;
+	D3D12_RESOURCE_STATES DestResourceState;
 	uint32 DestOffset;
 	uint32 Size;
 	ECopyType CopyType;
diff --git a/Engine/Source/Runtime/D3D12RHI/Private/D3D12RHI.cpp b/Engine/Source/Runtime/D3D12RHI/Private/D3D12RHI.cpp
--- a/Engine/Source/Runtime/D3D12RHI/Private/D3D12RHI.cpp
+++ b/Engine/Source/Runtime/D3D12RHI/Private/D3D12RHI.cpp
@@ -241,6 +241,9 @@ FD3D12DynamicRHI::FD3D12DynamicRHI(const TArray<TSharedPtr<FD3D12Adapter>>& Chos
 	// All D3D12 hardware supports binding UAVs to Vertex Shaders.
 	// Enable run-time support if corresponding bit is set in DDSPI.
 	GRHIGlobals.SupportsVertexShaderUAVs = true;
+
+	GRHIGlobals.NeedExtraTransitions = true;
+	GRHIGlobals.EnableHighLevelTransitionPatching = true;
 }
 
 void FD3D12DynamicRHI::PostInit()
diff --git a/Engine/Source/Runtime/D3D12RHI/Private/D3D12RHIPrivate.h b/Engine/Source/Runtime/D3D12RHI/Private/D3D12RHIPrivate.h
--- a/Engine/Source/Runtime/D3D12RHI/Private/D3D12RHIPrivate.h
+++ b/Engine/Source/Runtime/D3D12RHI/Private/D3D12RHIPrivate.h
@@ -224,8 +224,6 @@ private:
 	FProcessResult ProcessSubmissionQueue();
 	FProcessResult ProcessInterruptQueue();
 
-	static void GenerateBarrierCommandListAndUpdateState(FD3D12CommandList* SourceCommandList, TArray<FD3D12CommandList*, TInlineAllocator<2>>& OutBarrierCommandLists);
-
 	FCriticalSection ObjectsToDeleteCS;
 	TArray<FD3D12DeferredDeleteObject> ObjectsToDelete;
 
@@ -654,26 +652,24 @@ private:
 	FD3D12Resource* const Resource;
 	FD3D12ResourceLocation* const ResourceLocation;
 	D3D12_RESOURCE_STATES const DesiredState;
+	D3D12_RESOURCE_STATES  BeforeState;
 	uint32                const Subresource;
 
 	bool bRestoreState = false;
 
 public:
-	FScopedResourceBarrier(FD3D12ContextCommon& Context, FD3D12Resource* Resource, FD3D12ResourceLocation* InResourceLocation, D3D12_RESOURCE_STATES DesiredState, uint32 Subresource)
+	FScopedResourceBarrier(FD3D12ContextCommon& Context, FD3D12Resource* Resource, FD3D12ResourceLocation* InResourceLocation, D3D12_RESOURCE_STATES DesiredState, uint32 Subresource, D3D12_RESOURCE_STATES BeforeState = D3D12_RESOURCE_STATE_TBD)
 		: Context     (Context)
 		, Resource(Resource)
 		, ResourceLocation(InResourceLocation)
 		, DesiredState(DesiredState)
+		, BeforeState(BeforeState)
 		, Subresource (Subresource)
 	{
-		if (Resource->RequiresResourceStateTracking())
-		{
-			Context.TransitionResource(Resource, D3D12_RESOURCE_STATE_TBD, DesiredState, Subresource);
-		}
-		else
+		if ((!Resource->RequiresResourceStateTracking()) || (BeforeState != D3D12_RESOURCE_STATE_TBD))
 		{
 			// Resources that do not require state tracking are kept in a default state, meaning the previous state for the transition is known.
-			D3D12_RESOURCE_STATES CurrentState = Resource->GetDefaultResourceState();
+			D3D12_RESOURCE_STATES CurrentState = (BeforeState != D3D12_RESOURCE_STATE_TBD) ? BeforeState : Resource->GetDefaultResourceState();
 
 			// Some states such as D3D12_RESOURCE_STATE_GENERIC_READ already includes D3D12_RESOURCE_STATE_COPY_SOURCE as well as other states, therefore transition isn't required.
 			if (CurrentState != DesiredState && !EnumHasAllFlags(CurrentState, DesiredState))
@@ -690,14 +686,7 @@ public:
 	{
 		if (bRestoreState)
 		{
-			// Only untracked resources should get restored.
-			check(!Resource->RequiresResourceStateTracking());
-
-			Context.AddTransitionBarrier(Resource, DesiredState, Resource->GetDefaultResourceState(), Subresource);
-		}
-		else if (Resource->RequiresResourceStateTracking() && ResourceLocation)
-		{
-			Context.ConditionalClearShaderResource(ResourceLocation, EShaderParameterTypeMask::SRVMask | EShaderParameterTypeMask::UAVMask);
+			Context.AddTransitionBarrier(Resource, DesiredState, (BeforeState != D3D12_RESOURCE_STATE_TBD) ? BeforeState : Resource->GetDefaultResourceState(), Subresource);
 		}
 	}
 };
diff --git a/Engine/Source/Runtime/D3D12RHI/Private/D3D12RayTracing.cpp b/Engine/Source/Runtime/D3D12RHI/Private/D3D12RayTracing.cpp
--- a/Engine/Source/Runtime/D3D12RHI/Private/D3D12RayTracing.cpp
+++ b/Engine/Source/Runtime/D3D12RHI/Private/D3D12RayTracing.cpp
@@ -748,7 +748,7 @@ FD3D12RayTracingCompactionRequestHandler::FD3D12RayTracingCompactionRequestHandl
 	ID3D12ResourceAllocator* ResourceAllocator = nullptr;
 	bool bHasInitialData = false;
 	PostBuildInfoBuffer = GetParentDevice()->GetParentAdapter()->CreateRHIBuffer(PostBuildInfoBufferDesc, 8,
-		FRHIBufferDesc(PostBuildInfoBufferDesc.Width, 0, BUF_UnorderedAccess | BUF_SourceCopy), ED3D12ResourceStateMode::MultiState, D3D12_RESOURCE_STATE_UNORDERED_ACCESS, bHasInitialData, GPUMask, ResourceAllocator, TEXT("PostBuildInfoBuffer"));
+		FRHIBufferDesc(PostBuildInfoBufferDesc.Width, 0, BUF_UnorderedAccess | BUF_SourceCopy), ED3D12ResourceStateMode::MultiState, D3D12_RESOURCE_STATE_UNORDERED_ACCESS, bHasInitialData, GPUMask, ResourceAllocator, TEXT("PostBuildInfoBuffer"), ERHIAccess::UAVMask);
 	SetName(PostBuildInfoBuffer->GetResource(), TEXT("PostBuildInfoBuffer"));
 
 	PostBuildInfoStagingBuffer = RHICreateStagingBuffer();
@@ -852,7 +852,8 @@ void FD3D12RayTracingCompactionRequestHandler::Update(FD3D12CommandContext& Cont
 		PostBuildInfoDesc.DestBuffer = PostBuildInfoBuffer->ResourceLocation.GetGPUVirtualAddress();
 		PostBuildInfoDesc.InfoType = D3D12_RAYTRACING_ACCELERATION_STRUCTURE_POSTBUILD_INFO_COMPACTED_SIZE;
 
-		Context.TransitionResource(PostBuildInfoBuffer->GetResource(), D3D12_RESOURCE_STATE_TBD, D3D12_RESOURCE_STATE_UNORDERED_ACCESS, 0);
+		//PostBuildInfoBuffer enters in D3D12_RESOURCE_STATE_UNORDERED_ACCESS
+		//Context.TransitionResource(PostBuildInfoBuffer->GetResource(), D3D12_RESOURCE_STATE_COPY_SOURCE, D3D12_RESOURCE_STATE_UNORDERED_ACCESS, 0);
 
 		// Force UAV barrier to make sure all previous builds ops are finished
 		Context.AddUAVBarrier();
@@ -862,10 +863,11 @@ void FD3D12RayTracingCompactionRequestHandler::Update(FD3D12CommandContext& Cont
 		Context.RayTracingCommandList()->EmitRaytracingAccelerationStructurePostbuildInfo(&PostBuildInfoDesc, ActiveBLASGPUAddresses.Num(), ActiveBLASGPUAddresses.GetData());
 
 		// Transition to copy source and perform the copy to readback
-		Context.TransitionResource(PostBuildInfoBuffer.GetReference()->GetResource(), D3D12_RESOURCE_STATE_UNORDERED_ACCESS, D3D12_RESOURCE_STATE_COPY_SOURCE, 0);
+		Context.TransitionResource(PostBuildInfoBuffer->GetResource(), D3D12_RESOURCE_STATE_UNORDERED_ACCESS, D3D12_RESOURCE_STATE_COPY_SOURCE, 0);
 		Context.FlushResourceBarriers();
 
 		Context.RHICopyToStagingBuffer(PostBuildInfoBuffer, PostBuildInfoStagingBuffer, 0, sizeof(uint64) * ActiveBLASGPUAddresses.Num());
+		Context.TransitionResource(PostBuildInfoBuffer->GetResource(), D3D12_RESOURCE_STATE_COPY_SOURCE, D3D12_RESOURCE_STATE_UNORDERED_ACCESS, 0);
 
 		// Update the sync point
 		PostBuildInfoBufferReadbackSyncPoint = Context.GetContextSyncPoint();
@@ -1670,7 +1672,7 @@ public:
 		ID3D12ResourceAllocator* ResourceAllocator = nullptr;
 		Buffer = Adapter->CreateRHIBuffer(
 			BufferDesc, BufferDesc.Alignment, FRHIBufferDesc(BufferDesc.Width, 0, BUF_Static), ED3D12ResourceStateMode::MultiState,
-			D3D12_RESOURCE_STATE_COPY_DEST, bHasInitialData, GPUMask, ResourceAllocator, TEXT("Shader binding table"));
+			D3D12_RESOURCE_STATE_COPY_DEST, bHasInitialData, GPUMask, ResourceAllocator, TEXT("Shader binding table"), ERHIAccess::CopyDest);
 
 		// Use copy queue for uploading the data
 		Context.BatchedSyncPoints.ToWait.Emplace(Buffer->UploadResourceDataViaCopyQueue(&Data));
@@ -2275,18 +2277,6 @@ public:
 		}
 	}
 
-	void TransitionResources(FD3D12CommandContext& CommandContext) const
-	{	
-		for (FD3D12ShaderResourceView* SRV : WorkerData[0].TransitionSRVs)
-		{
-			CommandContext.TransitionResource(SRV, D3D12_RESOURCE_STATE_NON_PIXEL_SHADER_RESOURCE);
-		}
-
-		for (FD3D12UnorderedAccessView* UAV : WorkerData[0].TransitionUAVs)
-		{
-			CommandContext.TransitionResource(UAV, D3D12_RESOURCE_STATE_UNORDERED_ACCESS);
-		}
-	}
 
 	struct alignas(PLATFORM_CACHE_LINE_SIZE) FWorkerThreadData
 	{
@@ -3334,7 +3324,7 @@ static TRefCountPtr<FD3D12Buffer> CreateRayTracingBuffer(FD3D12Adapter* Adapter,
 			BufferDesc, D3D12_RAYTRACING_ACCELERATION_STRUCTURE_BYTE_ALIGNMENT,
 			FRHIBufferDesc(BufferDesc.Width, 0, BUF_AccelerationStructure),
 			ED3D12ResourceStateMode::SingleState, D3D12_RESOURCE_STATE_RAYTRACING_ACCELERATION_STRUCTURE, bHasInitialData,
-			GPUMask, ResourceAllocator, *DebugNameString);
+			GPUMask, ResourceAllocator, *DebugNameString, ERHIAccess::BVHWrite);
 	}
 	else if (Type == ERayTracingBufferType::Scratch)
 	{
@@ -3345,7 +3335,7 @@ static TRefCountPtr<FD3D12Buffer> CreateRayTracingBuffer(FD3D12Adapter* Adapter,
 			BufferDesc, D3D12_RAYTRACING_ACCELERATION_STRUCTURE_BYTE_ALIGNMENT,
 			FRHIBufferDesc(BufferDesc.Width, 0, BUF_UnorderedAccess),
 			ED3D12ResourceStateMode::Default, D3D12_RESOURCE_STATE_UNORDERED_ACCESS, bHasInitialData,
-			GPUMask, ResourceAllocator, *DebugNameString);
+			GPUMask, ResourceAllocator, *DebugNameString, ERHIAccess::BVHWrite);
 
 		// Elevates the scratch buffer heap priority, which may help performance / stability in low memory conditions 
 		// (Acceleration structure already boosted from allocation side)
@@ -3818,39 +3808,6 @@ bool FD3D12RayTracingGeometry::BuffersValid(uint32 GPUIndex) const
 	return true;
 }
 
-void FD3D12RayTracingGeometry::TransitionBuffers(FD3D12CommandContext& CommandContext)
-{
-	// Transition vertex and index resources..
-	if (Initializer.IndexBuffer)
-	{
-		FD3D12Buffer* IndexBuffer = CommandContext.RetrieveObject<FD3D12Buffer>(Initializer.IndexBuffer.GetReference());
-		if (IndexBuffer->GetResource()->RequiresResourceStateTracking())
-		{
-			CommandContext.TransitionResource(
-				IndexBuffer->GetResource(),
-				D3D12_RESOURCE_STATE_TBD,
-				D3D12_RESOURCE_STATE_NON_PIXEL_SHADER_RESOURCE,
-				0
-			);
-		}
-	}
-
-	for (const FRayTracingGeometrySegment& Segment : Initializer.Segments)
-	{
-		const FBufferRHIRef& RHIVertexBuffer = Segment.VertexBuffer;
-		FD3D12Buffer* VertexBuffer = CommandContext.RetrieveObject<FD3D12Buffer>(RHIVertexBuffer.GetReference());
-		if (VertexBuffer->GetResource()->RequiresResourceStateTracking())
-		{
-			CommandContext.TransitionResource(
-				VertexBuffer->GetResource(),
-				D3D12_RESOURCE_STATE_TBD,
-				D3D12_RESOURCE_STATE_NON_PIXEL_SHADER_RESOURCE,
-				0
-			);
-		}
-	}
-}
-
 void FD3D12RayTracingGeometry::UpdateResidency(FD3D12CommandContext& CommandContext)
 {
 	if (Initializer.IndexBuffer)
@@ -4357,12 +4314,6 @@ void BuildAccelerationStructure(FD3D12CommandContext& CommandContext,
 	CommandContext.UpdateResidency(ScratchBuffer->GetResource());
 
 	// Enqueue transition to UAV/SRV
-	CommandContext.TransitionResource(
-		InstanceBuffer->GetResource(),
-		D3D12_RESOURCE_STATE_TBD,
-		D3D12_RESOURCE_STATE_NON_PIXEL_SHADER_RESOURCE,
-		0
-	);
 
 	if (ShouldRunRayTracingGPUValidation())
 	{
@@ -4553,27 +4504,6 @@ void FD3D12CommandContext::RHIBuildAccelerationStructures(const TArrayView<const
 		}
 	}
 
-	// Transition all VBs and IBs to readable state
-
-	for (const FRayTracingGeometryBuildParams& P : Params)
-	{
-		FD3D12RayTracingGeometry* Geometry = FD3D12DynamicRHI::ResourceCast(P.Geometry.GetReference());
-		Geometry->TransitionBuffers(*this);
-	}
-
-	{
-		FD3D12Buffer* ScratchBuffer = FD3D12DynamicRHI::ResourceCast(ScratchBufferRange.Buffer, GetGPUIndex());
-		if (ScratchBuffer->GetResource()->RequiresResourceStateTracking())
-		{
-			TransitionResource(
-				ScratchBuffer->GetResource(),
-				D3D12_RESOURCE_STATE_TBD,
-				D3D12_RESOURCE_STATE_UNORDERED_ACCESS,
-				0
-			);
-		}
-	}
-
 	FlushResourceBarriers();
 		
 	const uint32 GPUIndex = GetGPUIndex();
@@ -4813,16 +4743,6 @@ struct FD3D12RayTracingGlobalResourceBinder
 		CommandContext.UpdateResidency(UniformBuffer->ResourceLocation.GetResource());
 	}
 
-	void AddResourceTransition(FD3D12ShaderResourceView* SRV)
-	{
-		CommandContext.TransitionResource(SRV, D3D12_RESOURCE_STATE_NON_PIXEL_SHADER_RESOURCE);
-	}
-
-	void AddResourceTransition(FD3D12UnorderedAccessView* UAV)
-	{
-		CommandContext.TransitionResource(UAV, D3D12_RESOURCE_STATE_UNORDERED_ACCESS);
-	}
-
 	void AddRayTracingSceneReference(FD3D12RayTracingScene* D3D12RayTracingScene)
 	{
 		D3D12RayTracingScene->UpdateResidency(CommandContext);
@@ -5080,7 +5000,6 @@ static bool SetRayTracingShaderResources(
 			}
 
 			Binder.AddReferencedShaderResource(UAV->GetShaderResource());
-			Binder.AddResourceTransition(UAV);
 		}
 
 		void SetSRV(FRHIShaderResourceView* RHISRV, uint8 Index)
@@ -5099,7 +5018,6 @@ static bool SetRayTracingShaderResources(
 			}
 
 			Binder.AddReferencedShaderResource(SRV->GetShaderResource());
-			Binder.AddResourceTransition(SRV);
 
 			FD3D12RayTracingScene* ReferencedRayTracingScene = SRV->GetRayTracingScene();
 			if (ReferencedRayTracingScene)
@@ -5128,7 +5046,6 @@ static bool SetRayTracingShaderResources(
 			}
 
 			Binder.AddReferencedTexture(RHITexture);
-			Binder.AddResourceTransition(SRV);
 		}
 
 		void SetResourceCollection(FRHIResourceCollection* ResourceCollection, uint8 Index)
@@ -5150,7 +5067,6 @@ static bool SetRayTracingShaderResources(
 			BoundSRVMask |= 1ull << Index;
 
 			Binder.AddReferencedShaderResource(SRV->GetShaderResource());
-			Binder.AddResourceTransition(SRV);
 		}
 
 		void SetSampler(FRHISamplerState* RHISampler, uint8 Index)
@@ -5454,8 +5370,8 @@ static void DispatchRays(FD3D12CommandContext& CommandContext,
 		DispatchRaysDescBuffer = Device->GetRayTracingDispatchRaysDescBuffer(QueueType);
 		FD3D12Resource* DispatchRaysDescBufferResource = DispatchRaysDescBuffer->GetResource();
 
-		CommandContext.TransitionResource(DispatchRaysDescBufferResource, D3D12_RESOURCE_STATE_TBD, D3D12_RESOURCE_STATE_COPY_DEST, 0);
-		CommandContext.TransitionResource(ArgumentBuffer->GetResource(), D3D12_RESOURCE_STATE_TBD, D3D12_RESOURCE_STATE_COPY_SOURCE, 0);
+		CommandContext.TransitionResource(DispatchRaysDescBufferResource, D3D12_RESOURCE_STATE_INDIRECT_ARGUMENT, D3D12_RESOURCE_STATE_COPY_DEST, 0);
+		CommandContext.TransitionResource(ArgumentBuffer->GetResource(), D3D12_RESOURCE_STATE_NON_PIXEL_SHADER_RESOURCE|D3D12_RESOURCE_STATE_INDIRECT_ARGUMENT, D3D12_RESOURCE_STATE_COPY_SOURCE, 0);
 		CommandContext.FlushResourceBarriers();
 
 		// Compute the allocation & copy sizes
@@ -5496,6 +5412,12 @@ static void DispatchRays(FD3D12CommandContext& CommandContext,
 			0
 		);
 
+		CommandContext.TransitionResource(ArgumentBuffer->GetResource(), 
+			D3D12_RESOURCE_STATE_COPY_SOURCE,
+			D3D12_RESOURCE_STATE_NON_PIXEL_SHADER_RESOURCE | D3D12_RESOURCE_STATE_INDIRECT_ARGUMENT, 
+			0
+		);
+
 		CommandContext.FlushResourceBarriers();
 	}
 
@@ -5569,11 +5491,6 @@ static void DispatchRays(FD3D12CommandContext& CommandContext,
 
 	if (bResourcesBound)
 	{
-		if (OptShaderTable)
-		{
-			OptShaderTable->TransitionResources(CommandContext);
-		}
-
 		CommandContext.FlushResourceBarriers();
 
 		ID3D12StateObject* RayTracingStateObject = nullptr;
diff --git a/Engine/Source/Runtime/D3D12RHI/Private/D3D12RayTracing.h b/Engine/Source/Runtime/D3D12RHI/Private/D3D12RayTracing.h
--- a/Engine/Source/Runtime/D3D12RHI/Private/D3D12RayTracing.h
+++ b/Engine/Source/Runtime/D3D12RHI/Private/D3D12RayTracing.h
@@ -37,7 +37,6 @@ public:
 	}
 
 	void SetupHitGroupSystemParameters(uint32 InGPUIndex);
-	void TransitionBuffers(FD3D12CommandContext& CommandContext);
 	void UpdateResidency(FD3D12CommandContext& CommandContext);
 	void CompactAccelerationStructure(FD3D12CommandContext& CommandContext, uint32 InGPUIndex, uint64 InSizeAfterCompaction);
 	void CreateAccelerationStructureBuildDesc(FD3D12CommandContext& CommandContext, EAccelerationStructureBuildMode BuildMode, D3D12_GPU_VIRTUAL_ADDRESS ScratchBufferAddress,
diff --git a/Engine/Source/Runtime/D3D12RHI/Private/D3D12RenderTarget.cpp b/Engine/Source/Runtime/D3D12RHI/Private/D3D12RenderTarget.cpp
--- a/Engine/Source/Runtime/D3D12RHI/Private/D3D12RenderTarget.cpp
+++ b/Engine/Source/Runtime/D3D12RHI/Private/D3D12RenderTarget.cpp
@@ -191,8 +191,6 @@ void FD3D12CommandContext::ResolveTextureUsingShader(
 			}
 		#endif
 
-			TransitionResource(DestTextureDSV, D3D12_RESOURCE_STATE_DEPTH_WRITE);
-
 			FlushResourceBarriers();
 
 			GraphicsCommandList()->ClearDepthStencilView(DestTextureDSV->GetOfflineCpuHandle(), D3D12_CLEAR_FLAG_DEPTH | D3D12_CLEAR_FLAG_STENCIL, 0, 0, 0, nullptr);
@@ -224,8 +222,6 @@ void FD3D12CommandContext::ResolveTextureUsingShader(
 			}
 		#endif
 
-			TransitionResource(DestTextureRTV, D3D12_RESOURCE_STATE_RENDER_TARGET);
-
 			FlushResourceBarriers();
 
 			FLinearColor ClearColor(0, 0, 0, 0);
@@ -360,8 +356,13 @@ void FD3D12CommandContext::ResolveTexture(UE::RHICore::FResolveTextureInfo Info)
 			int32 DestSubresource   = CalcSubresource(Info.MipLevel, ArraySlice, DestDesc.NumMips);
 			int32 SourceSubresource = CalcSubresource(Info.MipLevel, ArraySlice, SourceDesc.NumMips);
 
-			FScopedResourceBarrier ConditionalScopeResourceBarrierDst(*this, DestTexture->GetResource(),   &DestTexture->ResourceLocation,   D3D12_RESOURCE_STATE_RESOLVE_DEST,   DestSubresource);
-			FScopedResourceBarrier ConditionalScopeResourceBarrierSrc(*this, SourceTexture->GetResource(), &SourceTexture->ResourceLocation, D3D12_RESOURCE_STATE_RESOLVE_SOURCE, SourceSubresource);
+			// Source resource is in RTV, transition to D3D12_RESOURCE_STATE_RESOLVE_SOURCE and go back to D3D12_RESOURCE_STATE_RENDER_TARGET
+			// TODO: high level setting to D3D12_RESOURCE_STATE_RESOLVE_SOURCE
+			FScopedResourceBarrier ConditionalScopeResourceBarrierSrc(*this, SourceTexture->GetResource(), &SourceTexture->ResourceLocation, D3D12_RESOURCE_STATE_RESOLVE_SOURCE, SourceSubresource, D3D12_RESOURCE_STATE_RENDER_TARGET);
+
+			// Dest resource is already in D3D12_RESOURCE_STATE_RESOLVE_DEST, leave the FScopedResourceBarrier in case of using single state resources
+			FScopedResourceBarrier ConditionalScopeResourceBarrierDst(*this, DestTexture->GetResource(), &DestTexture->ResourceLocation, D3D12_RESOURCE_STATE_RESOLVE_DEST, DestSubresource);
+
 
 			FlushResourceBarriers();
 			GraphicsCommandList()->ResolveSubresource(DestResource->GetResource(), DestSubresource, SourceTexture->GetResource()->GetResource(), SourceSubresource, DestFormatTypeless);
@@ -467,11 +468,14 @@ TRefCountPtr<FD3D12Resource> FD3D12DynamicRHI::GetStagingTexture(FRHITexture* Te
 		// Remember the width, height, pitch, etc...
 		ReadbackHeapDesc = DestFootprint;
 
+		// TODO: transition only subresource not the whole resource
+		RHICmdList.TransitionInternal(FRHITransitionInfo(TextureRHI, ERHIAccess::Unknown, ERHIAccess::CopySrc, EResourceTransitionFlags::IgnoreAfterState));
+
 		RHICmdList.EnqueueLambda([&](FRHICommandList& ExecutingCmdList)
 		{
 			FD3D12CommandContext& Context = FD3D12CommandContext::Get(ExecutingCmdList, GPUIndex);
 
-			FScopedResourceBarrier ScopeResourceBarrierSource(Context, Texture->GetResource(), &Texture->ResourceLocation, D3D12_RESOURCE_STATE_COPY_SOURCE, SourceCopyLocation.SubresourceIndex);
+			//FScopedResourceBarrier ScopeResourceBarrierSource(Context, Texture->GetResource(), &Texture->ResourceLocation, D3D12_RESOURCE_STATE_COPY_SOURCE, SourceCopyLocation.SubresourceIndex);
 			Context.FlushResourceBarriers();
 			// Upload heap doesn't need to transition
 
@@ -484,6 +488,9 @@ TRefCountPtr<FD3D12Resource> FD3D12DynamicRHI::GetStagingTexture(FRHITexture* Te
 			Context.UpdateResidency(Texture->GetResource());
 		});
 
+		RHICmdList.TransitionInternal(FRHITransitionInfo(TextureRHI, ERHIAccess::CopySrc, ERHIAccess::Unknown, EResourceTransitionFlags::IgnoreAfterState));
+
+
 		// We need to execute the command list so we can read the data from readback heap
 		RHICmdList.SubmitAndBlockUntilGPUIdle();
 
diff --git a/Engine/Source/Runtime/D3D12RHI/Private/D3D12Resources.cpp b/Engine/Source/Runtime/D3D12RHI/Private/D3D12Resources.cpp
--- a/Engine/Source/Runtime/D3D12RHI/Private/D3D12Resources.cpp
+++ b/Engine/Source/Runtime/D3D12RHI/Private/D3D12Resources.cpp
@@ -98,6 +98,7 @@ FD3D12Resource::FD3D12Resource(FD3D12Device* ParentDevice,
 	, bRequiresResidencyTracking(bool(ENABLE_RESIDENCY_MANAGEMENT))
 	, bDepthStencil(false)
 	, bDeferDelete(true)
+	, bReconcileInitialState(false)
 {
 #if UE_BUILD_DEBUG
 	FPlatformAtomics::InterlockedIncrement(&TotalResourceCount);
@@ -771,6 +772,7 @@ HRESULT FD3D12Adapter::CreateCommittedResource(const FD3D12ResourceDesc& InDesc,
 #endif
 
 	HRESULT hr = S_OK;
+	bool bReconcileInitialState = false;
 #if INTEL_EXTENSIONS
 	if (InDesc.bRequires64BitAtomicSupport && IsRHIDeviceIntel() && GDX12INTCAtomicUInt64Emulation)
 	{
@@ -810,6 +812,7 @@ HRESULT FD3D12Adapter::CreateCommittedResource(const FD3D12ResourceDesc& InDesc,
 		{
 			// Change the initial state to match the initial layout on construction
 			InInitialState = D3D12_RESOURCE_STATE_COMMON;
+			bReconcileInitialState = true;
 		}
 	}
 	else
@@ -827,6 +830,8 @@ HRESULT FD3D12Adapter::CreateCommittedResource(const FD3D12ResourceDesc& InDesc,
 		SetName(*ppOutResource, Name);
 
 		(*ppOutResource)->StartTrackingForResidency();
+		(*ppOutResource)->SetInitalStateReconcile(bReconcileInitialState);
+		
 
 		TraceMemoryAllocation(*ppOutResource);
 	}
@@ -951,6 +956,7 @@ HRESULT FD3D12Adapter::CreatePlacedResource(const FD3D12ResourceDesc& InDesc, FD
 	ID3D12Heap* Heap = BackingHeap->GetHeap();
 
 	TRefCountPtr<ID3D12Resource> pResource;
+	bool bReconcileInitialState = false;
 	HRESULT hr = S_OK;
 #if INTEL_EXTENSIONS
 	if (InDesc.bRequires64BitAtomicSupport && IsRHIDeviceIntel() && GDX12INTCAtomicUInt64Emulation)
@@ -989,6 +995,7 @@ HRESULT FD3D12Adapter::CreatePlacedResource(const FD3D12ResourceDesc& InDesc, FD
 		{
 			// Change the initial state to match the initial layout on construction
 			InInitialState = D3D12_RESOURCE_STATE_COMMON;
+			bReconcileInitialState = true;
 		}
 	}
 	else
@@ -1011,6 +1018,8 @@ HRESULT FD3D12Adapter::CreatePlacedResource(const FD3D12ResourceDesc& InDesc, FD
 			InDesc,
 			BackingHeap,
 			HeapDesc.Properties.Type);
+		
+		(*ppOutResource)->SetInitalStateReconcile(bReconcileInitialState);
 
 #if PLATFORM_WINDOWS
 		if (IsTrackingAllAllocations() && BackingHeap->GetHeapDesc().Properties.Type == D3D12_HEAP_TYPE_DEFAULT)
@@ -1488,7 +1497,7 @@ void FD3D12ResourceLocation::AsStandAlone(FD3D12Resource* Resource, uint64 InSiz
 }
 
 
-bool FD3D12ResourceLocation::OnAllocationMoved(FD3D12ContextArray const& Contexts, FRHIPoolAllocationData* InNewData)
+bool FD3D12ResourceLocation::OnAllocationMoved(FD3D12ContextArray const& Contexts, FRHIPoolAllocationData* InNewData, D3D12_RESOURCE_STATES& OutCreateState)
 {
 	// Assume linked list allocated for now - only defragging allocator
 	FRHIPoolAllocationData& AllocationData = GetPoolAllocatorPrivateData().PoolData;
@@ -1564,6 +1573,8 @@ bool FD3D12ResourceLocation::OnAllocationMoved(FD3D12ContextArray const& Context
 		VERIFYD3D12RESULT(CurrentResource->GetParentDevice()->GetParentAdapter()->CreatePlacedResource(CurrentResource->GetDesc(), HeapAndOffset.Heap, HeapAndOffset.Offset, CreateState, ResourceStateMode, D3D12_RESOURCE_STATE_TBD, ClearValue, &NewResource, *Name.ToString()));
 		UnderlyingResource = NewResource;
 		GPUVirtualAddress = UnderlyingResource->GetGPUVirtualAddress() + OffsetFromBaseOfResource;
+		
+		OutCreateState = CreateState == D3D12_RESOURCE_STATE_TBD ? CurrentResource->GetWritableState() : CreateState;
 	}
 
 	// Refresh aliases
@@ -1583,6 +1594,25 @@ bool FD3D12ResourceLocation::OnAllocationMoved(FD3D12ContextArray const& Context
 	// Notify all the dependent resources about the change
 	Owner->ResourceRenamed(Contexts);
 
+	if (OutCreateState == D3D12_RESOURCE_STATE_TBD)
+	{
+		ERHIAccess ResourceAccess = ERHIAccess::Unknown;
+		FD3D12Texture* CurrentResourceRHITexture = nullptr;
+
+		if (CurrentResource->GetDesc().Dimension == D3D12_RESOURCE_DIMENSION_BUFFER)
+		{
+			FD3D12Buffer* CurrentResourceRHIBuffer = (FD3D12Buffer*)Owner;
+			ResourceAccess = CurrentResourceRHIBuffer->GetTrackedAccess_Unsafe();
+		}
+		else
+		{
+			CurrentResourceRHITexture = (FD3D12Texture*)Owner;
+			ResourceAccess = CurrentResourceRHITexture->GetTrackedAccess_Unsafe();
+		}
+
+		OutCreateState = ResourceAccess == ERHIAccess::Unknown ? CurrentResource->GetInitialResourceState() : GetD3D12ResourceState(ResourceAccess, ED3D12QueueType::Direct, CurrentResourceRHITexture);
+	}
+
 	return true;
 }
 
@@ -1703,6 +1733,12 @@ void FD3D12ResourceBarrierBatcher::FlushIntoCommandList(FD3D12CommandList& Comma
 			InsertTimestamp(true);
 		}
 
+#if DEBUG_RESOURCE_STATES
+		#define DX12_RESOURCE_NAME_TO_LOG	"Test"
+   		TArrayView<FD3D12ResourceBarrier> SubsetBarriers(Barriers.GetData() + BatchStart, BatchEnd - BatchStart);
+   		TConstArrayView<D3D12_RESOURCE_BARRIER> ConstArrayView(	reinterpret_cast<const D3D12_RESOURCE_BARRIER*>(SubsetBarriers.GetData()), SubsetBarriers.Num());
+ 		LogResourceBarriers(ConstArrayView, CommandList.Interfaces.CommandList.GetReference(), FString(DX12_RESOURCE_NAME_TO_LOG));
+#endif
 		CommandList.GraphicsCommandList()->ResourceBarrier(BatchEnd - BatchStart, &Barriers[BatchStart]);
 #if DEBUG_RESOURCE_STATES
 		// Keep track of all the resource barriers that have been submitted to the current command list.
diff --git a/Engine/Source/Runtime/D3D12RHI/Private/D3D12Resources.h b/Engine/Source/Runtime/D3D12RHI/Private/D3D12Resources.h
--- a/Engine/Source/Runtime/D3D12RHI/Private/D3D12Resources.h
+++ b/Engine/Source/Runtime/D3D12RHI/Private/D3D12Resources.h
@@ -204,7 +204,7 @@ private:
 #endif
 
 	const FD3D12ResourceDesc Desc;
-	CResourceState ResourceState;
+	D3D12_RESOURCE_STATES InitialResourceState{ D3D12_RESOURCE_STATE_TBD };
 	D3D12_RESOURCE_STATES DefaultResourceState{ D3D12_RESOURCE_STATE_TBD };
 	D3D12_RESOURCE_STATES ReadableState{ D3D12_RESOURCE_STATE_CORRUPT };
 	D3D12_RESOURCE_STATES WritableState{ D3D12_RESOURCE_STATE_CORRUPT };
@@ -218,10 +218,11 @@ private:
 	int32 NumMapCalls = 0;
 	uint16 SubresourceCount{};
 	uint8 PlaneCount;
-	bool bRequiresResourceStateTracking : 1;
-	bool bRequiresResidencyTracking : 1;
-	bool bDepthStencil : 1;
-	bool bDeferDelete : 1;
+	bool bRequiresResourceStateTracking : 1;	//1st bit 
+	bool bRequiresResidencyTracking : 1;		//2nd bit
+	bool bDepthStencil : 1;						//3rd bit
+	bool bDeferDelete : 1;						//4th bit
+	bool bReconcileInitialState : 1;			//5th bit
 
 #if UE_BUILD_DEBUG
 	static int64 TotalResourceCount;
@@ -313,13 +314,7 @@ public:
 	uint16 GetArraySize() const { return (Desc.Dimension == D3D12_RESOURCE_DIMENSION_TEXTURE3D) ? 1 : Desc.DepthOrArraySize; }
 	uint8 GetPlaneCount() const { return PlaneCount; }
 	uint16 GetSubresourceCount() const { return SubresourceCount; }
-	CResourceState& GetResourceState_OnResource()
-	{
-		check(bRequiresResourceStateTracking);
-		// This state is used as the resource's "global" state between command lists. It's only needed for resources that
-		// require state tracking.
-		return ResourceState;
-	}
+	D3D12_RESOURCE_STATES GetInitialResourceState() const { return InitialResourceState; }
 	D3D12_RESOURCE_STATES GetDefaultResourceState() const { check(!bRequiresResourceStateTracking); return DefaultResourceState; }
 	D3D12_RESOURCE_STATES GetWritableState() const { return WritableState; }
 	D3D12_RESOURCE_STATES GetReadableState() const { return ReadableState; }
@@ -346,7 +341,10 @@ public:
 	{
 		return DebugName;
 	}
-	
+
+	inline bool NeedInitialStateReconcile() { return bReconcileInitialState; }
+	inline void SetInitalStateReconcile(bool NewValue) { bReconcileInitialState = NewValue; }
+
 	void DoNotDeferDelete()
 	{
 		bDeferDelete = false;
@@ -441,7 +439,7 @@ public:
 			// state then when a transition is required (will transition via scoped push/pop to requested state)
 			if (!bSRVOnly && InResourceState != ERHIAccess::Unknown && InResourceState != ERHIAccess::Discard)
 			{
-				return GetD3D12ResourceState(InResourceState, ED3D12QueueType::Direct);
+				return GetD3D12ResourceState(InResourceState, ED3D12QueueType::Direct, nullptr);
 			}
 			else
 			{
@@ -525,9 +523,7 @@ private:
 			check(InDefaultState != D3D12_RESOURCE_STATE_RAYTRACING_ACCELERATION_STRUCTURE && InInitialState != D3D12_RESOURCE_STATE_RAYTRACING_ACCELERATION_STRUCTURE);
 #endif // D3D12_RHI_RAYTRACING
 
-			// Only a few resources (~1%) actually need resource state tracking
-			ResourceState.Initialize(SubresourceCount);
-			ResourceState.SetResourceState(InInitialState);
+			InitialResourceState = InInitialState;
 		}
 	}
 
@@ -643,6 +639,7 @@ class FD3D12BaseShaderResource;
 // A very light-weight and cache friendly way of accessing a GPU resource
 class FD3D12ResourceLocation : public FRHIPoolResource, public FD3D12DeviceChild, public FNoncopyable
 {
+	friend class FD3D12PoolAllocator;
 public:
 	enum class ResourceLocationType : uint8
 	{
@@ -711,7 +708,7 @@ public:
 	FD3D12PoolAllocatorPrivateData&    GetPoolAllocatorPrivateData   ()       { return AllocatorData.PoolAllocatorPrivateData;               }
 
 	// Pool allocation specific functions
-	bool OnAllocationMoved(FD3D12ContextArray const& Contexts, FRHIPoolAllocationData* InNewData);
+	bool OnAllocationMoved(FD3D12ContextArray const& Contexts, FRHIPoolAllocationData* InNewData, D3D12_RESOURCE_STATES& OutCreateState);
 	void UnlockPoolData();
 
 	bool IsValid() const { return Type != ResourceLocationType::eUndefined; }
diff --git a/Engine/Source/Runtime/D3D12RHI/Private/D3D12StateCache.cpp b/Engine/Source/Runtime/D3D12RHI/Private/D3D12StateCache.cpp
--- a/Engine/Source/Runtime/D3D12RHI/Private/D3D12StateCache.cpp
+++ b/Engine/Source/Runtime/D3D12RHI/Private/D3D12StateCache.cpp
@@ -1036,7 +1036,8 @@ void FD3D12StateCache::ApplySamplers(const FD3D12RootSignature* const pRootSigna
 	SamplerHeap->SetNextSlot(SamplerHeapSlot);
 }
 
-#if ASSERT_RESOURCE_STATES
+//#if ASSERT_RESOURCE_STATES
+#if 0
 /** Determine if an two views intersect */
 static inline bool ResourceViewsIntersect(FD3D12View* pLeftView, FD3D12View* pRightView)
 {
@@ -1469,7 +1470,6 @@ void FD3D12StateCache::InternalSetIndexBuffer(FD3D12Resource* Resource)
 	if (Resource->RequiresResourceStateTracking())
 	{
 		check(Resource->GetSubresourceCount() == 1);
-		CmdContext.TransitionResource(Resource, D3D12_RESOURCE_STATE_TBD, D3D12_RESOURCE_STATE_INDEX_BUFFER, D3D12_RESOURCE_BARRIER_ALL_SUBRESOURCES);
 	}
 }
 
@@ -1561,12 +1561,6 @@ void FD3D12StateCache::SetShaderResourceView(EShaderFrequency ShaderFrequency, F
 
 void FD3D12StateCache::SetRenderTargets(uint32 NumSimultaneousRenderTargets, FD3D12RenderTargetView** RTArray, FD3D12DepthStencilView* DSTarget)
 {
-	// Update the depth stencil
-	if (DSTarget)
-	{
-		CmdContext.TransitionResource(DSTarget);
-	}
-
 	if (PipelineState.Graphics.CurrentDepthStencilTarget != DSTarget)
 	{
 		PipelineState.Graphics.CurrentDepthStencilTarget = DSTarget;
@@ -1583,7 +1577,6 @@ void FD3D12StateCache::SetRenderTargets(uint32 NumSimultaneousRenderTargets, FD3
 
 		if (RTV)
 		{
-			CmdContext.TransitionResource(RTV, D3D12_RESOURCE_STATE_RENDER_TARGET);
 			PipelineState.Graphics.CurrentNumberOfRenderTargets++;
 		}
 
diff --git a/Engine/Source/Runtime/D3D12RHI/Private/D3D12Submission.cpp b/Engine/Source/Runtime/D3D12RHI/Private/D3D12Submission.cpp
--- a/Engine/Source/Runtime/D3D12RHI/Private/D3D12Submission.cpp
+++ b/Engine/Source/Runtime/D3D12RHI/Private/D3D12Submission.cpp
@@ -396,7 +396,6 @@ FD3D12DynamicRHI::FProcessResult FD3D12DynamicRHI::ProcessSubmissionQueue()
 		{
 			while (true)
 			{
-				TArray<FD3D12Queue*, TInlineAllocator<GD3D12MaxNumQueues>> QueuesWithPayloads;
 				{
 					FD3D12Payload* Payload = CurrentQueue.PendingSubmission.Peek();
 					if (!Payload)
@@ -433,7 +432,6 @@ FD3D12DynamicRHI::FProcessResult FD3D12DynamicRHI::ProcessSubmissionQueue()
 
 					check(!CurrentQueue.PayloadToSubmit);
 					CurrentQueue.PayloadToSubmit = Payload;
-					QueuesWithPayloads.Add(&CurrentQueue);
 					Result.Status |= EQueueStatus::Processed;
 					bKickInterruptThread = true;
 
@@ -491,51 +489,11 @@ FD3D12DynamicRHI::FProcessResult FD3D12DynamicRHI::ProcessSubmissionQueue()
 					for (int32 Index = 0; Index < Payload->CommandListsToExecute.Num(); Index++)
 					{
 						FD3D12CommandList* CurrentCommandList = Payload->CommandListsToExecute[Index];
-						TArray<FD3D12CommandList*, TInlineAllocator<2>> BarrierCommandLists;
-						GenerateBarrierCommandListAndUpdateState(CurrentCommandList, BarrierCommandLists);
-
-						for (FD3D12CommandList* BarrierCommandList : BarrierCommandLists)
-						{
-							FD3D12Queue& BarrierQueue = BarrierCommandList->Device->GetQueue(BarrierCommandList->QueueType);
-
-							if (&BarrierQueue == &CurrentQueue)
-							{
-								// Barrier command list will run on the current queue.
-								// Insert it immediately before the corresponding command list that generated it.
-								check(BarrierQueue.PayloadToSubmit);
-								BarrierQueue.PayloadToSubmit->CommandListsToExecute.Insert(BarrierCommandList, Index++);
-							}
-							else
-							{
-								// Barrier command list will run on a different queue.
-								if (!BarrierQueue.PayloadToSubmit)
-								{
-									BarrierQueue.PayloadToSubmit = new FD3D12Payload(BarrierQueue);
-									QueuesWithPayloads.Add(&BarrierQueue);
-								}
-								
-								BarrierQueue.PayloadToSubmit->CommandListsToExecute.Add(BarrierCommandList);
-							}
-
-							// Append the barrier cmdlist begin/end timestamps to the other queue.
-							AccumulateQueries(BarrierCommandList);
-						}
-
 						AccumulateQueries(CurrentCommandList);
 					}
 				}
 
-				// Queues with work to submit other than the current one (CurrentQueue) are performing barrier operations.
-				// Submit this work first, followed by a fence signal + enqueued wait.
-				for (FD3D12Queue* OtherQueue : QueuesWithPayloads)
-				{
-					if (OtherQueue != &CurrentQueue)
-					{
-						uint64 ValueSignaled = OtherQueue->FinalizePayload(true, PayloadsToHandDown);
-						CurrentQueue.PayloadToSubmit->AddQueueFenceWait(OtherQueue->Fence, ValueSignaled);
-					}
-					FlushPayloads(FD3D12Queue::MaxBatchedPayloads);
-				}
+				FlushPayloads(FD3D12Queue::MaxBatchedPayloads);
 
 				// Now submit the original payload
 				CurrentQueue.FinalizePayload(false, PayloadsToHandDown);
@@ -554,164 +512,6 @@ FD3D12DynamicRHI::FProcessResult FD3D12DynamicRHI::ProcessSubmissionQueue()
 	return Result;
 }
 
-void FD3D12DynamicRHI::GenerateBarrierCommandListAndUpdateState(FD3D12CommandList* SourceCommandList, TArray<FD3D12CommandList*, TInlineAllocator<2>>& OutBarrierCommandLists)
-{
-	TRACE_CPUPROFILER_EVENT_SCOPE(GenerateBarrierCommandListAndUpdateState);
-
-	FD3D12ResourceBarrierBatcher Batcher;
-	FD3D12ResourceBarrierBatcher BatcherForGraphicsToAsync;
-
-#if ENABLE_RESIDENCY_MANAGEMENT
-	TArray<FD3D12Resource*> ResourcesToMakeResident;
-	ResourcesToMakeResident.Reserve(SourceCommandList->State.PendingResourceBarriers.Num());
-#endif // ENABLE_RESIDENCY_MANAGEMENT
-
-	bool bHasGraphicStates = false;
-	for (const FD3D12PendingResourceBarrier& PRB : SourceCommandList->State.PendingResourceBarriers)
-	{
-		// Should only be doing this for the few resources that need state tracking
-		check(PRB.Resource->RequiresResourceStateTracking());
-
-		CResourceState& ResourceState = PRB.Resource->GetResourceState_OnResource();
-
-		const D3D12_RESOURCE_STATES Before = ResourceState.GetSubresourceState(PRB.SubResource);
-		const D3D12_RESOURCE_STATES After = PRB.State;
-
-		// We shouldn't have any TBD / CORRUPT states here
-		check(Before != D3D12_RESOURCE_STATE_TBD && Before != D3D12_RESOURCE_STATE_CORRUPT);
-		check(After  != D3D12_RESOURCE_STATE_TBD && After  != D3D12_RESOURCE_STATE_CORRUPT);
-
-		if (Before != After)
-		{
-			FD3D12ResourceBarrierBatcher* CurrentBatcher = &Batcher;
-
-			if (SourceCommandList->QueueType != ED3D12QueueType::Direct)
-			{
-				check(!IsDirectQueueExclusiveD3D12State(After));
-				if (IsDirectQueueExclusiveD3D12State(Before))
-				{
-					// If we are transitioning to a subset of the already set state on the resource (SRV_MASK -> SRV_COMPUTE for example)
-					// and there are no other transitions done on this resource in the command list then this transition can be ignored
-					// (otherwise a transition from SRV_COMPUTE as before state might already be recorded in the command list)
-					CResourceState& ResourceState_OnCommandList = SourceCommandList->GetResourceState_OnCommandList(PRB.Resource);
-					if (ResourceState_OnCommandList.HasInternalTransition() || !EnumHasAllFlags(Before, After))
-					{
-						CurrentBatcher = &BatcherForGraphicsToAsync;
-					}
-					else
-					{
-						// should be the same final state as well
-						check(After == ResourceState_OnCommandList.GetSubresourceState(PRB.SubResource));
-						
-						// Force the original state again if we're skipping this transition
-						ResourceState_OnCommandList.SetSubresourceState(PRB.SubResource, Before);
-						continue;
-					}
-				}
-			}
-
-			if (PRB.Resource->IsBackBuffer() && EnumHasAnyFlags(After, BackBufferBarrierWriteTransitionTargets))
-			{
-				CurrentBatcher->AddTransition(PRB.Resource, Before, After, PRB.SubResource);
-			}
-			// Special case for UAV access resources transitioning from UAV (then they need to transition from the cache hidden state instead)
-			else if (PRB.Resource->GetUAVAccessResource() && EnumHasAnyFlags(Before | After, D3D12_RESOURCE_STATE_UNORDERED_ACCESS))
-			{
-				// After state should never be UAV here
-				check(!EnumHasAnyFlags(After, D3D12_RESOURCE_STATE_UNORDERED_ACCESS));
-					
-				// Add the aliasing barrier
-				CurrentBatcher->AddAliasingBarrier(PRB.Resource->GetUAVAccessResource(), PRB.Resource->GetResource());
-
-				D3D12_RESOURCE_STATES UAVState = ResourceState.GetUAVHiddenResourceState();
-				check(UAVState != D3D12_RESOURCE_STATE_TBD && !EnumHasAnyFlags(UAVState, D3D12_RESOURCE_STATE_UNORDERED_ACCESS));
-				if (UAVState != After)
-				{
-					CurrentBatcher->AddTransition(PRB.Resource, UAVState, After, PRB.SubResource);
-				}
-			}
-			else
-			{
-				CurrentBatcher->AddTransition(PRB.Resource, Before, After, PRB.SubResource);
-			}
-		}
-
-	#if ENABLE_RESIDENCY_MANAGEMENT
-		ResourcesToMakeResident.Add(PRB.Resource);
-	#endif // ENABLE_RESIDENCY_MANAGEMENT
-	}
-
-	// Update the tracked resource states with the final states from the command list
-	for (auto const& [Resource, CommandListState] : SourceCommandList->State.TrackedResourceState)
-	{
-		CResourceState& ResourceState = Resource->GetResourceState_OnResource();
-		if (CommandListState.AreAllSubresourcesSame())
-		{
-			D3D12_RESOURCE_STATES State = CommandListState.GetSubresourceState(0);
-			if (State != D3D12_RESOURCE_STATE_TBD)
-			{
-				ResourceState.SetResourceState(State);
-			}
-
-			D3D12_RESOURCE_STATES UAVHiddenState = CommandListState.GetUAVHiddenResourceState();
-			if (UAVHiddenState != D3D12_RESOURCE_STATE_TBD)
-			{
-				check(Resource->GetUAVAccessResource());
-				ResourceState.SetUAVHiddenResourceState(UAVHiddenState);
-			}
-		}
-		else
-		{
-			for (uint32 Subresource = 0; Subresource < Resource->GetSubresourceCount(); ++Subresource)
-			{
-				D3D12_RESOURCE_STATES State = CommandListState.GetSubresourceState(Subresource);
-				if (State != D3D12_RESOURCE_STATE_TBD)
-				{
-					ResourceState.SetSubresourceState(Subresource, State);
-				}
-			}
-		}
-	}
-
-	const auto CreateBarrierCommandList = [&] (FD3D12Queue& Queue, FD3D12ResourceBarrierBatcher& Batcher)
-	{
-		// This command list requires a separate barrier command list to fix up tracked resource states.
-		TRACE_CPUPROFILER_EVENT_SCOPE(GetResourceBarrierCommandList);
-
-		// Get an allocator if we don't have one
-		if (!Queue.BarrierAllocator)
-		{
-			Queue.BarrierAllocator = Queue.Device->ObtainCommandAllocator(Queue.QueueType);
-		}
-
-		// Get a new command list
-		FD3D12CommandList* BarrierCommandList = Queue.Device->ObtainCommandList(Queue.BarrierAllocator, &Queue.BarrierTimestamps, nullptr);
-
-	#if ENABLE_RESIDENCY_MANAGEMENT
-		for (FD3D12Resource* Resource : ResourcesToMakeResident)
-		{
-			BarrierCommandList->UpdateResidency(Resource);
-		}
-	#endif // ENABLE_RESIDENCY_MANAGEMENT
-
-		Batcher.FlushIntoCommandList(*BarrierCommandList, Queue.BarrierTimestamps);
-		BarrierCommandList->Close();
-		return BarrierCommandList;
-	};
-
-	if (Batcher.Num() > 0)
-	{
-		FD3D12Queue& Queue = SourceCommandList->Device->GetQueue(SourceCommandList->QueueType);
-		OutBarrierCommandLists.Emplace(CreateBarrierCommandList(Queue, Batcher));
-	}
-	
-	if (BatcherForGraphicsToAsync.Num() > 0)
-	{
-		FD3D12Queue& Queue = SourceCommandList->Device->GetQueue(ED3D12QueueType::Direct);
-		OutBarrierCommandLists.Emplace(CreateBarrierCommandList(Queue, BatcherForGraphicsToAsync));
-	}
-}
-
 uint64 FD3D12Queue::FinalizePayload(bool bRequiresSignal, FPayloadArray& PayloadsToHandDown)
 {
 	TRACE_CPUPROFILER_EVENT_SCOPE(ExecuteCommandList);
diff --git a/Engine/Source/Runtime/D3D12RHI/Private/D3D12Texture.cpp b/Engine/Source/Runtime/D3D12RHI/Private/D3D12Texture.cpp
--- a/Engine/Source/Runtime/D3D12RHI/Private/D3D12Texture.cpp
+++ b/Engine/Source/Runtime/D3D12RHI/Private/D3D12Texture.cpp
@@ -709,6 +709,28 @@ FD3D12Texture* FD3D12DynamicRHI::CreateD3D12Texture(const FRHITextureCreateDesc&
 
 		check(Location.IsValid());
 
+		{
+			FD3D12Resource* CreatedResource = Location.GetResource();
+		
+			if (CreatedResource->NeedInitialStateReconcile())
+			{
+				const D3D12_RESOURCE_STATES RealInitialState = CreatedResource->GetInitialResourceState();
+
+				if (InitialState != RealInitialState)
+				{
+					uint32 GPUIndex = Device->GetGPUIndex();
+					RHICmdList->EnqueueLambda([CreatedResource, RealInitialState, InitialState, GPUIndex](FRHICommandListBase& ExecutingCmdList)
+						{
+#if UE_BUILD_DEBUG
+							UE_LOG(LogD3D12RHI, Warning, TEXT("Resource %s - Initial State Transition %s -> %s"), *CreatedResource->GetName().ToString(), *ConvertToResourceStateString(RealInitialState), *ConvertToResourceStateString(InitialState));
+#endif
+							FD3D12CommandContext& Context = FD3D12CommandContext::Get(ExecutingCmdList, GPUIndex);
+							Context.AddTransitionBarrier(CreatedResource, RealInitialState, InitialState, D3D12_RESOURCE_BARRIER_ALL_SUBRESOURCES);
+						});
+				}
+			}
+		}
+
 		if (ResourceDesc.NeedsUAVAliasWorkarounds())
 		{
 			Adapter->CreateUAVAliasResource(ClearValuePtr, CreateDesc.DebugName, Location);
@@ -735,6 +757,9 @@ FD3D12Texture* FD3D12DynamicRHI::CreateD3D12Texture(const FRHITextureCreateDesc&
 		D3D12TextureOut->InitializeTextureData(*RHICmdList, CreateDesc, InitialState);
 		CreateDesc.BulkData->Discard();
 	}
+
+	check( (D3D12TextureOut->GetTrackedSubResourceCount()) == (D3D12TextureOut->ResourceLocation.GetResource()->GetSubresourceCount()) );
+
 	return D3D12TextureOut;
 #else
 	checkf(false, TEXT("XBOX_CODE_MERGE : Removed. The Xbox platform version should be used."));
@@ -1606,10 +1631,14 @@ FRHILockTextureResult FD3D12Texture::Lock(FRHICommandListBase& RHICmdList, const
 		CD3DX12_TEXTURE_COPY_LOCATION DestCopyLocation(StagingTexture->GetResource(), PlacedFootprint);
 		CD3DX12_TEXTURE_COPY_LOCATION SourceCopyLocation(GetResource()->GetResource(), Subresource);
 
+		if (RHICmdList.NeedsExtraTransitions())
+		{
+			RHICmdList.TransitionInternal(FRHITransitionInfo(this, ERHIAccess::Unknown, ERHIAccess::CopySrc, EResourceTransitionFlags::IgnoreAfterState, Arguments.MipIndex, 0, 0));
+		}
+
 		RHICmdList.EnqueueLambda([this, DestCopyLocation, SourceCopyLocation](FRHICommandListBase& ExecutingCmdList)
 		{
 			FD3D12CommandContext& Context = FD3D12CommandContext::Get(ExecutingCmdList, 0);
-			FScopedResourceBarrier ScopeResourceBarrierSource(Context, GetResource(), &ResourceLocation, D3D12_RESOURCE_STATE_COPY_SOURCE, SourceCopyLocation.SubresourceIndex);
 
 			Context.FlushResourceBarriers();
 			Context.GraphicsCommandList()->CopyTextureRegion(
@@ -1621,6 +1650,11 @@ FRHILockTextureResult FD3D12Texture::Lock(FRHICommandListBase& RHICmdList, const
 			Context.UpdateResidency(GetResource());
 		});
 
+		if (RHICmdList.NeedsExtraTransitions())
+		{		
+			RHICmdList.TransitionInternal(FRHITransitionInfo(this, ERHIAccess::CopySrc, ERHIAccess::Unknown, EResourceTransitionFlags::IgnoreAfterState, Arguments.MipIndex, 0 , 0));
+		}
+
 		// We need to execute the command list so we can read the data from the map below
 		RHICmdList.GetAsImmediate().SubmitAndBlockUntilGPUIdle();
 
@@ -1820,9 +1854,16 @@ void FD3D12Texture::Unlock(FRHICommandListBase& RHICmdList, const FRHILockTextur
 	}
 	else if (!LockedResource->bLockedForReadOnly)
 	{
+
+		if (RHICmdList.NeedsExtraTransitions())
+		{
+			RHICmdList.TransitionInternal(FRHITransitionInfo(this, ERHIAccess::Unknown, ERHIAccess::CopyDest, EResourceTransitionFlags::IgnoreAfterState, Arguments.MipIndex, ArrayIndex, 0));
+		}
+
 		RHICmdList.EnqueueLambda([
 			RootTexture = this,
 			Subresource,
+			ArrayIndex,
 			LockedResource = MoveTemp(LockedResource)
 		](FRHICommandListBase& ExecutingCmdList)
 		{
@@ -1841,6 +1882,12 @@ void FD3D12Texture::Unlock(FRHICommandListBase& RHICmdList, const FRHILockTextur
 				Texture.UpdateTexture(Context, Subresource, 0, 0, 0, SourceCopyLocation);
 			}
 		});
+
+		if (RHICmdList.NeedsExtraTransitions())
+		{
+			RHICmdList.TransitionInternal(FRHITransitionInfo(this, ERHIAccess::CopyDest, ERHIAccess::Unknown, EResourceTransitionFlags::IgnoreAfterState, Arguments.MipIndex, ArrayIndex, 0));
+		}
+
 	}
 }
 
@@ -1895,6 +1942,11 @@ void FD3D12Texture::UpdateTexture2D(FRHICommandListBase& RHICmdList, uint32 MipI
 
 	CD3DX12_TEXTURE_COPY_LOCATION SourceCopyLocation(UploadHeapResourceLocation.GetResource()->GetResource(), PlacedTexture2D);
 
+	if (RHICmdList.NeedsExtraTransitions())
+	{
+		RHICmdList.TransitionInternal(FRHITransitionInfo(this, ERHIAccess::Unknown, ERHIAccess::CopyDest, EResourceTransitionFlags::IgnoreAfterState, MipIndex, 0, 0));
+	}
+
 	RHICmdList.EnqueueLambda([
 		RootTexture = this,
 		SourceCopyLocation,
@@ -1908,6 +1960,11 @@ void FD3D12Texture::UpdateTexture2D(FRHICommandListBase& RHICmdList, uint32 MipI
 			Texture.UpdateTexture(Context, MipIndex, UpdateRegion.DestX, UpdateRegion.DestY, 0, SourceCopyLocation);
 		}
 	});
+
+	if (RHICmdList.NeedsExtraTransitions())
+	{
+		RHICmdList.TransitionInternal(FRHITransitionInfo(this, ERHIAccess::CopyDest, ERHIAccess::Unknown, EResourceTransitionFlags::IgnoreAfterState, MipIndex, 0, 0));
+	}
 }
 
 static void GetReadBackHeapDescImpl(D3D12_PLACED_SUBRESOURCE_FOOTPRINT& OutFootprint, ID3D12Device* InDevice, D3D12_RESOURCE_DESC const& InResourceDesc, uint32 InSubresource)
@@ -2175,7 +2232,20 @@ void FD3D12DynamicRHI::RHIUpdateTexture3D(FRHICommandListBase& RHICmdList, FRHIT
 		}
 	}
 
+	FD3D12UpdateTexture3DData* UpdateDataD3D12 = reinterpret_cast<FD3D12UpdateTexture3DData*>(&UpdateData.PlatformData[0]);
+	bool bNeedTransition = (!(UpdateDataD3D12->bComputeShaderCopy)) && RHICmdList.NeedsExtraTransitions();
+
+	if (bNeedTransition)
+	{
+		RHICmdList.TransitionInternal(FRHITransitionInfo(TextureRHI, ERHIAccess::Unknown, ERHIAccess::CopyDest, EResourceTransitionFlags::IgnoreAfterState, MipIndex, 0, 0));
+	}
+
 	EndUpdateTexture3D_Internal(RHICmdList, UpdateData);
+
+	if (bNeedTransition)
+	{
+		RHICmdList.TransitionInternal(FRHITransitionInfo(TextureRHI, ERHIAccess::CopyDest, ERHIAccess::Unknown, EResourceTransitionFlags::IgnoreAfterState, MipIndex, 0, 0));
+	}
 }
 
 
diff --git a/Engine/Source/Runtime/D3D12RHI/Private/D3D12UAV.cpp b/Engine/Source/Runtime/D3D12RHI/Private/D3D12UAV.cpp
--- a/Engine/Source/Runtime/D3D12RHI/Private/D3D12UAV.cpp
+++ b/Engine/Source/Runtime/D3D12RHI/Private/D3D12UAV.cpp
@@ -389,8 +389,6 @@ void FD3D12CommandContext::ClearUAV(TRHICommandList_RecursiveHazardous<FD3D12Com
 					Device->CopyDescriptorsSimple(1, DestSlot, OfflineHandle, D3D12_DESCRIPTOR_HEAP_TYPE_CBV_SRV_UAV);
 				}
 
-				Context.TransitionResource(UnorderedAccessView, D3D12_RESOURCE_STATE_UNORDERED_ACCESS);
-
 				Context.FlushResourceBarriers();
 				Context.GraphicsCommandList()->ClearUnorderedAccessViewUint(GPUHandle, OfflineHandle, UAV.GetResource()->GetResource(), *reinterpret_cast<const UINT(*)[4]>(ClearValues), 0, nullptr);
 				Context.UpdateResidency(UAV.GetResource());
diff --git a/Engine/Source/Runtime/D3D12RHI/Private/D3D12Util.cpp b/Engine/Source/Runtime/D3D12RHI/Private/D3D12Util.cpp
--- a/Engine/Source/Runtime/D3D12RHI/Private/D3D12Util.cpp
+++ b/Engine/Source/Runtime/D3D12RHI/Private/D3D12Util.cpp
@@ -86,6 +86,26 @@ void SetName(FD3D12Resource* const Resource, const TCHAR* const Name)
 #endif
 }
 
+FString GetD312ObjectName(ID3D12Object* const Object)
+{
+#if NAME_OBJECTS
+	if (Object == nullptr)
+		return FString("Unknown Resource");
+
+	TCHAR OutName[128];
+	UINT size = 128;
+	HRESULT hr = Object->GetPrivateData(WKPDID_D3DDebugObjectNameW, &size, OutName);
+	check(size <= 128);
+
+	if (hr != S_OK)
+		return FString("Unknown Resource");
+	
+	return FString(OutName);
+#else
+	return FString();
+#endif
+}
+
 static FString GetD3D12DeviceHungErrorString(HRESULT ErrorCode)
 {
 	FString ErrorCodeText;
@@ -1510,13 +1530,14 @@ FString ConvertToResourceStateString(uint32 ResourceState)
 	return ResourceStateString;
 }
 
-void LogResourceBarriers(TConstArrayView<D3D12_RESOURCE_BARRIER> Barriers, ID3D12CommandList* const pCommandList)
+void LogResourceBarriers(TConstArrayView<D3D12_RESOURCE_BARRIER> Barriers, ID3D12CommandList* const pCommandList, const FString& ResourceName)
 {
 	// Configure what resource barriers are logged.
-	const bool bLogAll = false;
+	const bool bLogAll = true;
 	const bool bLogTransitionDepth = true;
 	const bool bLogTransitionRenderTarget = true;
 	const bool bLogTransitionUAV = true;
+	const bool bCheckResourceName = ResourceName.IsEmpty() ? false : true;
 
 	// Create the state bit mask to indicate what barriers should be logged.
 	uint32 ShouldLogMask = bLogAll ? static_cast<uint32>(-1) : 0;
@@ -1548,21 +1569,41 @@ void LogResourceBarriers(TConstArrayView<D3D12_RESOURCE_BARRIER> Barriers, ID3D1
 
 			if (bShouldLog)
 			{
-				UE_LOG(LogD3D12RHI, Log, TEXT("*** BARRIER (CmdList: %016llX) %u/%u: %016llX (Sub: %u), %s -> %s"), pCommandList, i + 1, Barriers.Num(),
-					currentBarrier.Transition.pResource,
-					currentBarrier.Transition.Subresource,
-					*StateBefore,
-					*StateAfter);
+				const FString BarrierResourceName = GetD312ObjectName(currentBarrier.Transition.pResource);
+
+				if ((bCheckResourceName == false) || (BarrierResourceName == ResourceName))
+				{
+					UE_LOG(LogD3D12RHI, Log, TEXT("*** BARRIER (CmdList: %016llX) %u/%u: %s %016llX (Sub: %u), %s -> %s"), pCommandList, i + 1, Barriers.Num(),
+						*BarrierResourceName,
+						currentBarrier.Transition.pResource,
+						currentBarrier.Transition.Subresource,
+						*StateBefore,
+						*StateAfter);
+				}
 			}
 			break;
 		}
 
 		case D3D12_RESOURCE_BARRIER_TYPE_UAV:
-			UE_LOG(LogD3D12RHI, Log, TEXT("*** BARRIER (CmdList: %016llX) %u/%u: UAV Barrier"), pCommandList, i + 1, Barriers.Num());
+			{
+				const FString BarrierResourceName = GetD312ObjectName(currentBarrier.UAV.pResource);
+				if ((bCheckResourceName == false) || (BarrierResourceName == ResourceName))
+				{
+					UE_LOG(LogD3D12RHI, Log, TEXT("*** BARRIER (CmdList: %016llX) %u/%u: UAV Barrier %s"), pCommandList, i + 1, Barriers.Num(), *BarrierResourceName);
+				}
+			}
 			break;
 
 		case D3D12_RESOURCE_BARRIER_TYPE_ALIASING:
-			UE_LOG(LogD3D12RHI, Log, TEXT("*** BARRIER (CmdList: %016llX) %u/%u: Aliasing Barrier, %016llX -> %016llX"), pCommandList, i + 1, Barriers.Num(), currentBarrier.Aliasing.pResourceBefore, currentBarrier.Aliasing.pResourceAfter);
+			{
+				const FString BarrierResourceNameBefore = GetD312ObjectName(currentBarrier.Aliasing.pResourceBefore);
+				const FString BarrierResourceNameAfter = GetD312ObjectName(currentBarrier.Aliasing.pResourceAfter);
+				
+				if ((bCheckResourceName == false) || (BarrierResourceNameBefore == ResourceName) || (BarrierResourceNameAfter == ResourceName))
+				{
+					UE_LOG(LogD3D12RHI, Log, TEXT("*** BARRIER (CmdList: %016llX) %u/%u: Aliasing Barrier, %016llX %s -> %016llX %s"), pCommandList, i + 1, Barriers.Num(), currentBarrier.Aliasing.pResourceBefore, *BarrierResourceNameBefore, currentBarrier.Aliasing.pResourceAfter, *BarrierResourceNameAfter);
+				}
+			}
 			break;
 
 		default:
@@ -1573,26 +1614,41 @@ void LogResourceBarriers(TConstArrayView<D3D12_RESOURCE_BARRIER> Barriers, ID3D1
 }
 
 
-D3D12_RESOURCE_STATES GetD3D12ResourceState(ERHIAccess InRHIAccess, ED3D12QueueType QueueType)
+D3D12_RESOURCE_STATES GetD3D12ResourceState(ERHIAccess InRHIAccess, ED3D12QueueType QueueType, FD3D12Texture* InRHID3D12Texture)
 {
 	// Add switch for common states (should cover all writeable states)
 	switch (InRHIAccess)
 	{
-	// all single write states
+		// all single write states
 	case ERHIAccess::RTV:					return D3D12_RESOURCE_STATE_RENDER_TARGET;
-	case ERHIAccess::UAVMask:		
-	case ERHIAccess::UAVCompute:	
+	case ERHIAccess::UAVMask:
+	case ERHIAccess::UAVCompute:
 	case ERHIAccess::UAVGraphics:			return D3D12_RESOURCE_STATE_UNORDERED_ACCESS;
 	case ERHIAccess::DSVWrite:				return D3D12_RESOURCE_STATE_DEPTH_WRITE;
 	case ERHIAccess::CopyDest:				return D3D12_RESOURCE_STATE_COPY_DEST;
 	case ERHIAccess::ResolveDst:			return D3D12_RESOURCE_STATE_RESOLVE_DEST;
 	case ERHIAccess::Present:				return D3D12_RESOURCE_STATE_PRESENT;
 
-	// Generic read for mask read states
-	case ERHIAccess::ReadOnlyMask:	
+		// Generic read for mask read states
+	case ERHIAccess::ReadOnlyMask:
 	case ERHIAccess::ReadOnlyExclusiveMask:	return D3D12_RESOURCE_STATE_GENERIC_READ;
 	default:
 	{
+		D3D12_RESOURCE_STATES ExtraReadState = D3D12_RESOURCE_STATES(0);
+
+		if (InRHID3D12Texture)
+		{
+// 			if (InRHID3D12Texture->GetResource()->IsDepthStencilResource())
+// 			{
+// 				ExtraReadState |= D3D12_RESOURCE_STATE_DEPTH_READ;
+// 			}
+			 
+			if (InRHID3D12Texture->SkipsFastClearFinalize())
+			{
+				ExtraReadState |= FD3D12CommandContext::SkipFastClearEliminateState;
+			}
+		}
+
 		// Special case for DSV read & write (Depth write allows depth read as well in D3D)
 		if (InRHIAccess == ERHIAccess(ERHIAccess::DSVRead | ERHIAccess::DSVWrite))
 		{
@@ -1609,11 +1665,11 @@ D3D12_RESOURCE_STATES GetD3D12ResourceState(ERHIAccess InRHIAccess, ED3D12QueueT
 			// Translate the requested after state to a D3D state
 			if (EnumHasAnyFlags(InRHIAccess, ERHIAccess::SRVGraphics) && QueueType == ED3D12QueueType::Direct)
 			{
-				State |= D3D12_RESOURCE_STATE_PIXEL_SHADER_RESOURCE | D3D12_RESOURCE_STATE_NON_PIXEL_SHADER_RESOURCE;
+				State |= D3D12_RESOURCE_STATE_PIXEL_SHADER_RESOURCE | D3D12_RESOURCE_STATE_NON_PIXEL_SHADER_RESOURCE | ExtraReadState;
 			}
 			if (EnumHasAnyFlags(InRHIAccess, ERHIAccess::SRVCompute))
 			{
-				State |= D3D12_RESOURCE_STATE_NON_PIXEL_SHADER_RESOURCE;
+				State |= D3D12_RESOURCE_STATE_NON_PIXEL_SHADER_RESOURCE | ExtraReadState;
 			}
 			if (EnumHasAnyFlags(InRHIAccess, ERHIAccess::VertexOrIndexBuffer))
 			{
@@ -1635,12 +1691,16 @@ D3D12_RESOURCE_STATES GetD3D12ResourceState(ERHIAccess InRHIAccess, ED3D12QueueT
 			{
 				State |= D3D12_RESOURCE_STATE_DEPTH_READ;
 			}
-#if PLATFORM_SUPPORTS_VARIABLE_RATE_SHADING
-			if (EnumHasAnyFlags(InRHIAccess, ERHIAccess::ShadingRateSource) && GRHISupportsAttachmentVariableRateShading)
+			if (EnumHasAnyFlags(InRHIAccess, ERHIAccess::ShadingRateSource))
 			{
+				if (GRHISupportsAttachmentVariableRateShading == false)
+				{
+					UE_LOG(LogD3D12RHI, Error, TEXT("Resource state is D3D12_RESOURCE_STATE_SHADING_RATE_SOURCE but RHI does not support VRS."));
+				}
+#if PLATFORM_SUPPORTS_VARIABLE_RATE_SHADING
 				State |= D3D12_RESOURCE_STATE_SHADING_RATE_SOURCE;
-			}
 #endif
+			}
 
 			// Should have at least one valid state
 			check(State != D3D12_RESOURCE_STATE_COMMON);
@@ -1651,177 +1711,21 @@ D3D12_RESOURCE_STATES GetD3D12ResourceState(ERHIAccess InRHIAccess, ED3D12QueueT
 	}
 }
 
-//==================================================================================================================================
-// CResourceState
-// Tracking of per-resource or per-subresource state
-//==================================================================================================================================
-//----------------------------------------------------------------------------------------------------------------------------------
-void CResourceState::Initialize(uint32 SubresourceCount)
-{
-	check(0 == m_SubresourceState.Num());
-
-	// Allocate space for per-subresource tracking structures
-	check(SubresourceCount > 0);
-	m_SubresourceState.SetNumUninitialized(SubresourceCount);
-	check(m_SubresourceState.Num() == SubresourceCount);
-
-	// No internal transition yet
-	bHasInternalTransition = 0;
-
-	// All subresources start out in an unknown state
-	SetResourceState(D3D12_RESOURCE_STATE_TBD);
-
-	// Unknown hidden resource state
-	SetUAVHiddenResourceState(D3D12_RESOURCE_STATE_TBD);
-}
-
-//----------------------------------------------------------------------------------------------------------------------------------
-bool CResourceState::AreAllSubresourcesSame() const
-{
-	return m_AllSubresourcesSame && (m_ResourceState != D3D12_RESOURCE_STATE_TBD);
-}
-
-//----------------------------------------------------------------------------------------------------------------------------------
-bool CResourceState::CheckResourceState(D3D12_RESOURCE_STATES State) const
-{
-	if (m_AllSubresourcesSame)
-	{
-		return State == m_ResourceState;
-	}
-	else
-	{
-		// All subresources must be individually checked
-		const uint32 numSubresourceStates = m_SubresourceState.Num();
-		for (uint32 i = 0; i < numSubresourceStates; i++)
-		{
-			if (m_SubresourceState[i] != State)
-			{
-				return false;
-			}
-		}
-
-		return true;
-	}
-}
-
-//----------------------------------------------------------------------------------------------------------------------------------
-bool CResourceState::CheckResourceStateInitalized() const
-{
-	return m_SubresourceState.Num() > 0;
-}
-
-//----------------------------------------------------------------------------------------------------------------------------------
-D3D12_RESOURCE_STATES CResourceState::GetSubresourceState(uint32 SubresourceIndex) const
-{
-	if (m_AllSubresourcesSame)
-	{
-		return static_cast<D3D12_RESOURCE_STATES>(m_ResourceState);
-	}
-	else
-	{
-		check(SubresourceIndex < static_cast<uint32>(m_SubresourceState.Num()));
-		return m_SubresourceState[SubresourceIndex];
-	}
-}
-
-//----------------------------------------------------------------------------------------------------------------------------------
-
-bool CResourceState::CheckAllSubresourceSame()
-{
-	// already marked same?
-	if (m_AllSubresourcesSame)
-	{
-		return true;
-	}
-	else
-	{
-		D3D12_RESOURCE_STATES State = m_SubresourceState[0];
-
-		// All subresources must be individually checked
-		const uint32 numSubresourceStates = m_SubresourceState.Num();
-		for (uint32 i = 1; i < numSubresourceStates; i++)
-		{
-			if (m_SubresourceState[i] != State)
-			{
-				return false;
-			}
-		}
-
-		SetResourceState(State);
-
-		return true;
-	}
-}
-
-//----------------------------------------------------------------------------------------------------------------------------------
-void CResourceState::SetResourceState(D3D12_RESOURCE_STATES State)
-{
-	m_AllSubresourcesSame = 1;
-
-	// m_ResourceState is restricted to 31 bits.  Ensure State can be properly represented.
-	check((State & (1 << 31)) == 0);
-
-	m_ResourceState = *reinterpret_cast<uint32*>(&State);
-
-	// State is now tracked per-resource, so m_SubresourceState should not be read.
-#if UE_BUILD_DEBUG
-	const uint32 numSubresourceStates = m_SubresourceState.Num();
-	for (uint32 i = 0; i < numSubresourceStates; i++)
-	{
-		m_SubresourceState[i] = D3D12_RESOURCE_STATE_CORRUPT;
-	}
-#endif
-}
-
-//----------------------------------------------------------------------------------------------------------------------------------
-void CResourceState::SetSubresourceState(uint32 SubresourceIndex, D3D12_RESOURCE_STATES State)
-{
-	// If setting all subresources, or the resource only has a single subresource, set the per-resource state
-	if (SubresourceIndex == D3D12_RESOURCE_BARRIER_ALL_SUBRESOURCES ||
-		m_SubresourceState.Num() == 1)
-	{
-		SetResourceState(State);
-	}
-	else
-	{
-		check(SubresourceIndex < static_cast<uint32>(m_SubresourceState.Num()));
-
-		// If state was previously tracked on a per-resource level, then transition to per-subresource tracking
-		if (m_AllSubresourcesSame)
-		{
-			const uint32 numSubresourceStates = m_SubresourceState.Num();
-			for (uint32 i = 0; i < numSubresourceStates; i++)
-			{
-				m_SubresourceState[i] = static_cast<D3D12_RESOURCE_STATES>(m_ResourceState);
-			}
-
-			m_AllSubresourcesSame = 0;
-
-			// State is now tracked per-subresource, so m_ResourceState should not be read.
-#if UE_BUILD_DEBUG
-			m_ResourceState = D3D12_RESOURCE_STATE_CORRUPT;
-#endif
-		}
-
-		m_SubresourceState[SubresourceIndex] = State;
-	}
-}
-
 #if ASSERT_RESOURCE_STATES
 // Forward declarations are required for the template functions
-template bool AssertResourceState(ID3D12CommandList* pCommandList, FD3D12View* pView, const D3D12_RESOURCE_STATES& State);
-
-bool AssertResourceState(ID3D12CommandList* pCommandList, FD3D12View* pView, const D3D12_RESOURCE_STATES& State)
-{
-	// Check the view
-	if (!pView)
-	{
-		// No need to check null views
-		return true;
-	}
-
-	return AssertResourceState(pCommandList, pView->GetResource(), State, pView->GetViewSubresourceSubset());
-}
+// template bool AssertResourceState(ID3D12CommandList* pCommandList, FD3D12View* pView, const D3D12_RESOURCE_STATES& State);
+// 
+// bool AssertResourceState(ID3D12CommandList* pCommandList, FD3D12View* pView, const D3D12_RESOURCE_STATES& State)
+// {
+// 	// Check the view
+// 	if (!pView)
+// 	{
+// 		// No need to check null views
+// 		return true;
+// 	}
+// 
+// 	return AssertResourceState(pCommandList, pView->GetResource(), State, pView->GetViewSubresourceSubset());
+// }
 
 bool AssertResourceState(ID3D12CommandList* pCommandList, FD3D12Resource* pResource, const D3D12_RESOURCE_STATES& State, uint32 Subresource)
 {
diff --git a/Engine/Source/Runtime/D3D12RHI/Private/D3D12Util.h b/Engine/Source/Runtime/D3D12RHI/Private/D3D12Util.h
--- a/Engine/Source/Runtime/D3D12RHI/Private/D3D12Util.h
+++ b/Engine/Source/Runtime/D3D12RHI/Private/D3D12Util.h
@@ -53,6 +53,7 @@ class FD3D12Resource;
 
 void SetName(ID3D12Object* const Object, const TCHAR* const Name);
 void SetName(FD3D12Resource* const Resource, const TCHAR* const Name);
+FString GetD312ObjectName(ID3D12Object* const Object);
 
 enum EShaderVisibility
 {
@@ -254,7 +255,7 @@ private:
 
 void LogExecuteCommandLists(uint32 NumCommandLists, ID3D12CommandList *const *ppCommandLists);
 FString ConvertToResourceStateString(uint32 ResourceState);
-void LogResourceBarriers(TConstArrayView<D3D12_RESOURCE_BARRIER> Barriers, ID3D12CommandList *const pCommandList);
+void LogResourceBarriers(TConstArrayView<D3D12_RESOURCE_BARRIER> Barriers, ID3D12CommandList *const pCommandList, const FString& ResourceName);
 
 // Custom resource states
 // To Be Determined (TBD) means we need to fill out a resource barrier before the command list is executed.
@@ -271,68 +272,9 @@ static bool IsDirectQueueExclusiveD3D12State(D3D12_RESOURCE_STATES InState)
 	return EnumHasAnyFlags(InState, D3D12_RESOURCE_STATE_RENDER_TARGET | D3D12_RESOURCE_STATE_DEPTH_WRITE | D3D12_RESOURCE_STATE_DEPTH_READ | D3D12_RESOURCE_STATE_PIXEL_SHADER_RESOURCE);
 }
 
-D3D12_RESOURCE_STATES GetD3D12ResourceState(ERHIAccess InRHIAccess, ED3D12QueueType QueueType);
+class FD3D12Texture;
 
-//==================================================================================================================================
-// CResourceState
-// Tracking of per-resource or per-subresource state
-//==================================================================================================================================
-class CResourceState
-{
-public:
-	void Initialize(uint32 SubresourceCount);
-
-	bool AreAllSubresourcesSame() const;
-	bool CheckResourceState(D3D12_RESOURCE_STATES State) const;
-	bool CheckResourceStateInitalized() const;
-	D3D12_RESOURCE_STATES GetSubresourceState(uint32 SubresourceIndex) const;
-	bool CheckAllSubresourceSame();
-	void SetResourceState(D3D12_RESOURCE_STATES State);
-	void SetSubresourceState(uint32 SubresourceIndex, D3D12_RESOURCE_STATES State);
-
-	D3D12_RESOURCE_STATES GetUAVHiddenResourceState() const
-	{
-		return static_cast<D3D12_RESOURCE_STATES>(UAVHiddenResourceState);
-	}
-
-	void SetUAVHiddenResourceState(D3D12_RESOURCE_STATES InUAVHiddenResourceState)
-	{
-		// The hidden state can never include UAV
-		check(InUAVHiddenResourceState == D3D12_RESOURCE_STATE_TBD || !EnumHasAnyFlags(InUAVHiddenResourceState, D3D12_RESOURCE_STATE_UNORDERED_ACCESS));
-		check((InUAVHiddenResourceState & (1 << 31)) == 0);
-
-		UAVHiddenResourceState = (uint32)InUAVHiddenResourceState;
-	}
-
-	void SetHasInternalTransition()
-	{
-		bHasInternalTransition = 1;
-	}
-	bool HasInternalTransition() const
-	{
-		return bHasInternalTransition != 0;
-	}
-
-private:
-	// Only used if m_AllSubresourcesSame is 1.
-	// Bits defining the state of the full resource, bits are from D3D12_RESOURCE_STATES
-	uint32 m_ResourceState : 31;
-
-	// Set to 1 if m_ResourceState is valid.  In this case, all subresources have the same state
-	// Set to 0 if m_SubresourceState is valid.  In this case, each subresources may have a different state (or may be unknown)
-	uint32 m_AllSubresourcesSame : 1;
-
-	// Special resource state to track previous state before resource transitioned to UAV when the resource
-	// has a UAV aliasing resource so correct previous state can be found (only single state allowed)
-	uint32 UAVHiddenResourceState : 31;
-
-	// Was the resource used for another transition than the pending transition
-	uint32 bHasInternalTransition : 1;
-
-	// Only used if m_AllSubresourcesSame is 0.
-	// The state of each subresources.  Bits are from D3D12_RESOURCE_STATES.
-	TArray<D3D12_RESOURCE_STATES, TInlineAllocator<4>> m_SubresourceState;
-};
+D3D12_RESOURCE_STATES GetD3D12ResourceState(ERHIAccess InRHIAccess, ED3D12QueueType QueueType, FD3D12Texture* InRHID3D12Texture);
 
 /**
  * The base class of threadsafe reference counted objects.
@@ -638,10 +580,10 @@ static void Get4KTileShape(D3D12_TILE_SHAPE* pTileShape, DXGI_FORMAT DXGIFormat,
 
 #if ASSERT_RESOURCE_STATES
 class FD3D12View;
-class FD3D12ViewSubset;
+struct FD3D12ViewSubset;
 
-template <class TView>
-bool AssertResourceState(ID3D12CommandList* pCommandList, FD3D12View* pView, const D3D12_RESOURCE_STATES& State);
+// template <class TView>
+// bool AssertResourceState(ID3D12CommandList* pCommandList, FD3D12View* pView, const D3D12_RESOURCE_STATES& State);
 
 bool AssertResourceState(ID3D12CommandList* pCommandList, FD3D12Resource* pResource, const D3D12_RESOURCE_STATES& State, uint32 Subresource);
 bool AssertResourceState(ID3D12CommandList* pCommandList, FD3D12Resource* pResource, const D3D12_RESOURCE_STATES& State, const FD3D12ViewSubset& ViewSubset);
@@ -742,4 +684,51 @@ struct FD3D12ScopeNoLock
 public:
 	FD3D12ScopeNoLock(FCriticalSection* CritSec) { /* Do Nothing! */ }
 	~FD3D12ScopeNoLock() { /* Do Nothing! */ }
-};
\ No newline at end of file
+};
+
+
+FORCEINLINE_DEBUGGABLE bool CheckResourceStateCompatibility(D3D12_RESOURCE_STATES State, D3D12_RESOURCE_FLAGS Flags, FString& IncompatibilityReason)
+{
+	// Check for ALLOW_RENDER_TARGET compatibility
+	if (State & D3D12_RESOURCE_STATE_RENDER_TARGET)
+	{
+		if (!(Flags & D3D12_RESOURCE_FLAG_ALLOW_RENDER_TARGET))
+		{
+			IncompatibilityReason.Append(TEXT("D3D12_RESOURCE_STATE_RENDER_TARGET requires D3D12_RESOURCE_FLAG_ALLOW_RENDER_TARGET.\n"));
+			return false;
+		}
+	}
+
+	// Check for ALLOW_DEPTH_STENCIL compatibility
+	if (State & D3D12_RESOURCE_STATE_DEPTH_WRITE || State & D3D12_RESOURCE_STATE_DEPTH_READ)
+	{
+		if (!(Flags & D3D12_RESOURCE_FLAG_ALLOW_DEPTH_STENCIL))
+		{
+			IncompatibilityReason.Append(TEXT("D3D12_RESOURCE_STATE_DEPTH_WRITE or D3D12_RESOURCE_STATE_DEPTH_READ requires D3D12_RESOURCE_FLAG_ALLOW_DEPTH_STENCIL.\n"));
+			return false;
+		}
+	}
+
+	// Check for ALLOW_UNORDERED_ACCESS compatibility
+	if (State & D3D12_RESOURCE_STATE_UNORDERED_ACCESS)
+	{
+		if (!(Flags & D3D12_RESOURCE_FLAG_ALLOW_UNORDERED_ACCESS))
+		{
+			IncompatibilityReason.Append(TEXT("D3D12_RESOURCE_STATE_UNORDERED_ACCESS requires D3D12_RESOURCE_FLAG_ALLOW_UNORDERED_ACCESS.\n"));
+			return false;
+		}
+	}
+
+	// Check for DENY_SHADER_RESOURCE incompatibility
+	if (State & (D3D12_RESOURCE_STATE_PIXEL_SHADER_RESOURCE | D3D12_RESOURCE_STATE_NON_PIXEL_SHADER_RESOURCE))
+	{
+		if (Flags & D3D12_RESOURCE_FLAG_DENY_SHADER_RESOURCE)
+		{
+			IncompatibilityReason.Append(TEXT("D3D12_RESOURCE_STATE_PIXEL_SHADER_RESOURCE or D3D12_RESOURCE_STATE_NON_PIXEL_SHADER_RESOURCE cannot be used with D3D12_RESOURCE_FLAG_DENY_SHADER_RESOURCE.\n"));
+			return false;
+		}
+	}
+
+	return true;
+}
+
diff --git a/Engine/Source/Runtime/D3D12RHI/Private/D3D12Viewport.cpp b/Engine/Source/Runtime/D3D12RHI/Private/D3D12Viewport.cpp
--- a/Engine/Source/Runtime/D3D12RHI/Private/D3D12Viewport.cpp
+++ b/Engine/Source/Runtime/D3D12RHI/Private/D3D12Viewport.cpp
@@ -86,8 +86,8 @@ FD3D12Texture* GetSwapChainSurface(FD3D12Device* Parent, EPixelFormat PixelForma
 	BackBufferDesc.bBackBuffer = true;
 
 	FString Name = FString::Printf(TEXT("BackBuffer%d"), BackBufferIndex);
-
-	ETextureCreateFlags SwapchainTextureCreateFlags = ETextureCreateFlags::RenderTargetable;
+	
+	ETextureCreateFlags SwapchainTextureCreateFlags = ETextureCreateFlags::RenderTargetable | ETextureCreateFlags::Presentable | ETextureCreateFlags::ResolveTargetable;
 #if D3D12RHI_SUPPORTS_UAV_BACKBUFFER
 	SwapchainTextureCreateFlags |= ETextureCreateFlags::UAV;
 #endif
@@ -133,9 +133,9 @@ FD3D12Texture* GetSwapChainSurface(FD3D12Device* Parent, EPixelFormat PixelForma
 				&NewTexture->ResourceLocation,
 				NewTexture,
 				PixelFormat,
-				TexCreate_RenderTargetable |  TexCreate_ShaderResource,
-				D3D12_RESOURCE_STATE_PRESENT,
-				TEXT("SwapChainSurface"));
+				TexCreate_RenderTargetable |  TexCreate_ShaderResource | TexCreate_Presentable,
+				InitialState,
+				*Name);
 		}
 
 		// active stereoscopy initialization
@@ -535,27 +535,6 @@ bool FD3D12Viewport::Present(FD3D12CommandContextBase& ContextBase, bool bLockTo
 		// Those are not necessarily the swap chain back buffer in case of multi-gpu
 		FD3D12Texture* DeviceBackBuffer = Context.RetrieveObject<FD3D12Texture, FRHITexture>(GetBackBuffer_RHIThread());
 
-		Context.TransitionResource(
-			DeviceBackBuffer->GetShaderResourceView()->GetResource(),
-			D3D12_RESOURCE_STATE_TBD,
-			D3D12_RESOURCE_STATE_PRESENT,
-			0
-		);
-
-#if D3D12RHI_USE_SDR_BACKBUFFER
-		if (CurrentBackBuffer_RHIThread->TextureSDR)
-		{
-			FD3D12Texture* DeviceSDRBackBuffer = Context.RetrieveObject<FD3D12Texture, FRHITexture>(GetSDRBackBuffer_RHIThread());
-
-			Context.TransitionResource(
-				DeviceSDRBackBuffer->GetShaderResourceView()->GetResource(),
-				D3D12_RESOURCE_STATE_TBD,
-				D3D12_RESOURCE_STATE_PRESENT,
-				0
-			);
-		}
-#endif
-
 		Context.FlushResourceBarriers();
 
 		// Currently, the swap chain Present() is called directly by the RHI thread.
diff --git a/Engine/Source/Runtime/D3D12RHI/Private/D3D12Viewport.h b/Engine/Source/Runtime/D3D12RHI/Private/D3D12Viewport.h
--- a/Engine/Source/Runtime/D3D12RHI/Private/D3D12Viewport.h
+++ b/Engine/Source/Runtime/D3D12RHI/Private/D3D12Viewport.h
@@ -93,6 +93,22 @@ public:
 		return CurrentBackBuffer_RHIThread->Texture;
 	}
 
+#if D3D12RHI_USE_SDR_BACKBUFFER
+	FRHITexture* GetOptionalSDRBackBuffer(FRHITexture* BackBufferTex) const override
+	{
+		for (const FBackBufferData& CurBackBuffer : BackBuffers)
+		{
+			if ( (FRHITexture*)CurBackBuffer.Texture.GetReference() == BackBufferTex)
+			{
+				return (FRHITexture*)CurBackBuffer.TextureSDR.GetReference();
+			}
+		}
+
+		return nullptr;
+	}
+#endif
+
+
 #if WITH_MGPU
 	uint32 GetNextPresentGPUIndex() const
 	{
diff --git a/Engine/Source/Runtime/D3D12RHI/Private/D3D12WorkGraph.cpp b/Engine/Source/Runtime/D3D12RHI/Private/D3D12WorkGraph.cpp
--- a/Engine/Source/Runtime/D3D12RHI/Private/D3D12WorkGraph.cpp
+++ b/Engine/Source/Runtime/D3D12RHI/Private/D3D12WorkGraph.cpp
@@ -758,7 +758,7 @@ void FD3D12CommandContext::DispatchWorkGraphShaderBundle(
 		D3D12_RESOURCE_DESC Desc = CD3DX12_RESOURCE_DESC::Buffer(LocalRootArgs.GetResourceDataSize(), D3D12_RESOURCE_FLAG_NONE);
 		FD3D12Buffer* RootArgBuffer = GetParentDevice()->GetParentAdapter()->CreateRHIBuffer(
 			Desc, 16, FRHIBufferDesc(Desc.Width, 0, BUF_Static), ED3D12ResourceStateMode::MultiState,
-			D3D12_RESOURCE_STATE_COPY_DEST, true, FRHIGPUMask::FromIndex(GetParentDevice()->GetGPUIndex()), nullptr, TEXT("BundleRecordBuffer"));
+			D3D12_RESOURCE_STATE_COPY_DEST, true, FRHIGPUMask::FromIndex(GetParentDevice()->GetGPUIndex()), nullptr, TEXT("BundleRecordBuffer"), ERHIAccess::CopyDest);
 
 		BatchedSyncPoints.ToWait.Emplace(RootArgBuffer->UploadResourceDataViaCopyQueue(&LocalRootArgs));
 		TransitionResource(RootArgBuffer->GetResource(), D3D12_RESOURCE_STATE_COPY_DEST, D3D12_RESOURCE_STATE_COMMON, 0);
@@ -794,16 +794,16 @@ void FD3D12CommandContext::DispatchWorkGraphShaderBundle(
 		{
 			ClearShaderResources(UAV, EShaderParameterTypeMask::SRVMask);
 		}
-
-		for (FD3D12ShaderResourceView* SRV : BinderOps[0].TransitionSRVs)
-		{
-			TransitionResource(SRV, D3D12_RESOURCE_STATE_NON_PIXEL_SHADER_RESOURCE);
-		}
-
-		for (FD3D12UnorderedAccessView* UAV : BinderOps[0].TransitionUAVs)
-		{
-			TransitionResource(UAV, D3D12_RESOURCE_STATE_UNORDERED_ACCESS);
-		}
+//TODO: check resource view
+// 		for (FD3D12ShaderResourceView* SRV : BinderOps[0].TransitionSRVs)
+// 		{
+// 			TransitionResource(SRV, D3D12_RESOURCE_STATE_NON_PIXEL_SHADER_RESOURCE);
+// 		}
+// 
+// 		for (FD3D12UnorderedAccessView* UAV : BinderOps[0].TransitionUAVs)
+// 		{
+// 			TransitionResource(UAV, D3D12_RESOURCE_STATE_UNORDERED_ACCESS);
+// 		}
 	}
 
 	FlushResourceBarriers();
diff --git a/Engine/Source/Runtime/D3D12RHI/Private/Windows/WindowsD3D12RHIDefinitions.h b/Engine/Source/Runtime/D3D12RHI/Private/Windows/WindowsD3D12RHIDefinitions.h
--- a/Engine/Source/Runtime/D3D12RHI/Private/Windows/WindowsD3D12RHIDefinitions.h
+++ b/Engine/Source/Runtime/D3D12RHI/Private/Windows/WindowsD3D12RHIDefinitions.h
@@ -26,7 +26,7 @@ typedef ID3D12GraphicsCommandList ID3D12CopyCommandList;
 
 #define USE_STATIC_ROOT_SIGNATURE			0
 #define D3D12RHI_USE_CONSTANT_BUFFER_VIEWS	0
-#define D3D12RHI_USE_DUMMY_BACKBUFFER		1
+#define D3D12RHI_USE_DUMMY_BACKBUFFER		0
 #define D3D12RHI_USE_SDR_BACKBUFFER			0
 #define D3D12RHI_SUPPORTS_UAV_BACKBUFFER	0
 #define D3D12RHI_USE_DXGI_COLOR_SPACE		1
diff --git a/Engine/Source/Runtime/D3D12RHI/Private/Windows/WindowsD3D12Viewport.cpp b/Engine/Source/Runtime/D3D12RHI/Private/Windows/WindowsD3D12Viewport.cpp
--- a/Engine/Source/Runtime/D3D12RHI/Private/Windows/WindowsD3D12Viewport.cpp
+++ b/Engine/Source/Runtime/D3D12RHI/Private/Windows/WindowsD3D12Viewport.cpp
@@ -232,6 +232,7 @@ void FD3D12Viewport::ConditionalResetSwapChain(bool bIgnoreFocus)
 	}
 }
 
+#if D3D12RHI_USE_DUMMY_BACKBUFFER
 /**
 * Create the dummy back buffer textures - They don't have actual D3D resource but are used to always reference the current back buffer index on the RHI thread
 */
@@ -250,6 +251,7 @@ FD3D12Texture* FD3D12Viewport::CreateDummyBackBufferTextures(FD3D12Adapter* InAd
 		});
 	return Result;
 }
+#endif
 
 void FD3D12Viewport::ResizeInternal()
 {
@@ -324,7 +326,6 @@ void FD3D12Viewport::ResizeInternal()
 	// Create dummy back buffer which always reference to the actual RHI thread back buffer - can't be bound directly to D3D12
 	DummyBackBuffer_RenderThread = CreateDummyBackBufferTextures(Adapter, PixelFormat, SizeX, SizeY);
 #else
-	static_assert(false, "D3D12RHI_USE_DUMMY_BACKBUFFER==0 has not been tested for viewports");
 	SetBackBufferIndex_RenderThread(0);
 #endif
 }
diff --git a/Engine/Source/Runtime/RHI/Private/RHICommandList.cpp b/Engine/Source/Runtime/RHI/Private/RHICommandList.cpp
--- a/Engine/Source/Runtime/RHI/Private/RHICommandList.cpp
+++ b/Engine/Source/Runtime/RHI/Private/RHICommandList.cpp
@@ -2283,17 +2283,18 @@ void FRHICommandListBase::TrackResourceState(TArrayView<const FRHITransition*> T
 			}
 	
 			// Some RHI will not transition to initial state on After = Discard when running on the Async queue
+			bool bDiscardOnAsync = false;
 			if (NeedsExtraTransitions())
 			{
 				if ((TransitionInfo->AccessAfter == ERHIAccess::Discard)
 					&& ((TransitionCreateInfoData->SrcPipelines == ERHIPipeline::AsyncCompute) || (TransitionCreateInfoData->SrcPipelines == ERHIPipeline::All))
 					&& (TransitionCreateInfoData->DstPipelines == ERHIPipeline::AsyncCompute))
 				{
-					TransitionInfo->Flags |= EResourceTransitionFlags::IgnoreAfterState;
+					bDiscardOnAsync = true;
 				}
 			}
-
-			const bool bNeedUpdateAfterState = (!((EnumHasAnyFlags(TransitionInfo->Flags, EResourceTransitionFlags::IgnoreAfterState)) || (TransitionInfo->AccessAfter == ERHIAccess::Unknown)));
+			const bool bShouldSkipAfterState = ((EnumHasAnyFlags(TransitionInfo->Flags, EResourceTransitionFlags::IgnoreAfterState)) || (TransitionInfo->AccessAfter == ERHIAccess::Unknown));
+			const bool bNeedUpdateAfterState = (bShouldSkipAfterState == false) && (bDiscardOnAsync == false);
 
 			if (bNeedUpdateAfterState)
 			{
diff --git a/Engine/Source/Runtime/RHI/Public/RHICommandList.h b/Engine/Source/Runtime/RHI/Public/RHICommandList.h
--- a/Engine/Source/Runtime/RHI/Public/RHICommandList.h
+++ b/Engine/Source/Runtime/RHI/Public/RHICommandList.h
@@ -872,9 +872,16 @@ public:
 			FRHITextureCreateDesc NewCreateDesc(CreateDesc);
 			NewCreateDesc.SetInitialState(RHIGetDefaultResourceState(CreateDesc.Flags, CreateDesc.BulkData != nullptr));
 
+#if RHI_TRANSITIONS_DEBUG_INITIAL_STATE
+			UE_LOG(LogRHI, Display, TEXT("Resource: %s Creation State %s"), CreateDesc.DebugName, *GetRHIAccessName(NewCreateDesc.InitialState));
+#endif
+
 			return GDynamicRHI->RHICreateTexture(*this, NewCreateDesc);
 		}
 
+#if RHI_TRANSITIONS_DEBUG_INITIAL_STATE
+		UE_LOG(LogRHI, Display, TEXT("Resource: %s Creation State %s"), CreateDesc.DebugName, *GetRHIAccessName(CreateDesc.InitialState));
+#endif
 		return GDynamicRHI->RHICreateTexture(*this, CreateDesc);
 	}
 
@@ -1154,10 +1161,17 @@ public:
 		return GRHIGlobals.NeedsShaderUnbinds;
 	}
 
-	// Returns true if the RHI needs extra transitions
+	// Returns true if the RHI needs and allow extra transitions
 	bool NeedsExtraTransitions() const
 	{
-		return GRHIGlobals.NeedExtraTransitions;
+		return GRHIGlobals.NeedExtraTransitions && bAllowExtraTransition;
+	}
+	// Returns old state of bAllowExtraTransition
+	bool SetAllowExtraTransition(bool NewState)
+	{
+		bool OldState = bAllowExtraTransition;
+		bAllowExtraTransition = NewState;
+		return OldState;
 	}
 
 	FRHIBatchedShaderParametersAllocator* CreateBatchedShaderParameterAllocator(ERHIBatchedShaderParameterAllocatorPageSize PageSize)
@@ -1261,6 +1275,7 @@ protected:
 	bool bUsesSetTrackedAccess   = false;
 	bool bUsesShaderBundles      = false;
 	bool bUsesLockFence          = false;
+	bool bAllowExtraTransition	 = true;
 
 	// The currently selected pipelines that RHI commands are directed to, during command list recording.
 	// This is also adjusted during command list execution based on recorded use of ActivatePipeline().
diff --git a/Engine/Source/Runtime/RHI/Public/RHICommandList.inl b/Engine/Source/Runtime/RHI/Public/RHICommandList.inl
--- a/Engine/Source/Runtime/RHI/Public/RHICommandList.inl
+++ b/Engine/Source/Runtime/RHI/Public/RHICommandList.inl
@@ -244,17 +244,19 @@ FORCEINLINE_DEBUGGABLE void FRHICommandListBase::PatchTransitionsInPlace(FRHITra
 	if (TransitionViewableResource)
 	{
 		// Some RHI will not transition to initial state on After = Discard when running on the Async queue
+		bool bDiscardOnAsync = false;
 		if (NeedsExtraTransitions())
 		{
 			if ((TransitionInfo->AccessAfter == ERHIAccess::Discard)
 				&& ((SrcPipeline == ERHIPipeline::AsyncCompute) || (SrcPipeline == ERHIPipeline::All))
 				&& (DestPipeline == ERHIPipeline::AsyncCompute))
 			{
-				TransitionInfo->Flags |= EResourceTransitionFlags::IgnoreAfterState;
+				bDiscardOnAsync = true;
 			}
 		}
-
-		const bool bNeedUpdateAfterState = !((EnumHasAnyFlags(TransitionInfo->Flags, EResourceTransitionFlags::IgnoreAfterState)) || (TransitionInfo->AccessAfter == ERHIAccess::Unknown));
+		const bool bShouldSkipAfterState = ((EnumHasAnyFlags(TransitionInfo->Flags, EResourceTransitionFlags::IgnoreAfterState)) || (TransitionInfo->AccessAfter == ERHIAccess::Unknown));
+		const bool bNeedUpdateAfterState = (bShouldSkipAfterState == false) && (bDiscardOnAsync == false);
+		
 		// Some RHI needs extra information on transition from Discard
 		const bool bBeforeIsDiscard = NeedsExtraTransitions() && (TransitionInfo->AccessBefore == ERHIAccess::Discard);
 
@@ -276,10 +278,10 @@ FORCEINLINE_DEBUGGABLE void FRHICommandListBase::PatchTransitionsInPlace(FRHITra
 			TransitionViewableResource->TrackedAccess.SetSubresourceState(SubResourceIndex, TransitionInfo->AccessAfter);
 			TransitionViewableResource->TrackedAccess.CheckAllSubresourceSame();
 		}
+		
+		check(TransitionInfo->AccessBefore != ERHIAccess::Unknown);
+		check(TransitionInfo->AccessAfter != ERHIAccess::Unknown);
 	}
-
-	check(TransitionInfo->AccessBefore != ERHIAccess::Unknown);
-	check(TransitionInfo->AccessAfter != ERHIAccess::Unknown);
 }
 
 
@@ -290,9 +292,12 @@ FORCEINLINE_DEBUGGABLE void FRHICommandListBase::PatchTransitionsInPlace(TArrayV
 	for (const FRHITransition* Transition : Transitions)
 	{
 		FRHITransitionCreateInfoData* TransitionCreateInfoData = Transition->CreationInfoData;
-		for (int32 i = 0; i < TransitionCreateInfoData->TransitionInfos.Num(); i++)
+		if (TransitionCreateInfoData)
 		{
-			PatchTransitionsInPlace((FRHITransitionInfo*)&TransitionCreateInfoData->TransitionInfos[i], TransitionCreateInfoData->SrcPipelines, TransitionCreateInfoData->DstPipelines);
+			for (int32 i = 0; i < TransitionCreateInfoData->TransitionInfos.Num(); i++)
+			{
+				PatchTransitionsInPlace((FRHITransitionInfo*)&TransitionCreateInfoData->TransitionInfos[i], TransitionCreateInfoData->SrcPipelines, TransitionCreateInfoData->DstPipelines);
+			}
 		}
 	}
 }
diff --git a/Engine/Source/Runtime/RHI/Public/RHIDefinitions.h b/Engine/Source/Runtime/RHI/Public/RHIDefinitions.h
--- a/Engine/Source/Runtime/RHI/Public/RHIDefinitions.h
+++ b/Engine/Source/Runtime/RHI/Public/RHIDefinitions.h
@@ -1611,10 +1611,6 @@ struct FShaderCodeValidationUBSize
 	uint32 Size;
 };
 
-#if UE_BUILD_DEBUG
-#define RHI_TRANSITIONS_HL_DEBUG	1
-#else
-#define RHI_TRANSITIONS_HL_DEBUG	0
-#endif
+#define RHI_TRANSITIONS_DEBUG_INITIAL_STATE	0
 
 #define RHI_UnknownSubResourceCount (0)
\ No newline at end of file
diff --git a/Engine/Source/Runtime/RHI/Public/RHIResources.h b/Engine/Source/Runtime/RHI/Public/RHIResources.h
--- a/Engine/Source/Runtime/RHI/Public/RHIResources.h
+++ b/Engine/Source/Runtime/RHI/Public/RHIResources.h
@@ -1724,7 +1724,7 @@ struct FRHITextureDesc
 	uint16 GetSubresourceCount() const
 	{
 		uint16 NumFaces = ((Dimension == ETextureDimension::TextureCube) || (Dimension == ETextureDimension::TextureCubeArray)) ? 6 : 1;
-		uint16 NumPlanes = IsStencilFormat(Format) ? 2 : 1;
+		uint16 NumPlanes = (IsStencilFormat(Format) || Format == PF_D24) ? 2 : 1;
 
 		return (ArraySize * (NumMips * NumFaces) * NumPlanes);
 	}
diff --git a/Engine/Source/Runtime/RenderCore/Private/RenderGraphBuilder.cpp b/Engine/Source/Runtime/RenderCore/Private/RenderGraphBuilder.cpp
--- a/Engine/Source/Runtime/RenderCore/Private/RenderGraphBuilder.cpp
+++ b/Engine/Source/Runtime/RenderCore/Private/RenderGraphBuilder.cpp
@@ -172,9 +172,8 @@ void EnumerateTextureAccess(FRDGParameterStruct PassParameters, ERDGPassFlags Pa
 
 				if (ResolveTexture && ResolveTexture != Texture)
 				{
-					// Resolve targets must use the RTV|ResolveDst flag combination when the resolve is performed through the render
-					// pass. The ResolveDst flag must be used alone only when the resolve is performed using RHICopyToResolveTarget.
-					AccessFunction(nullptr, ResolveTexture, ERHIAccess::RTV | ERHIAccess::ResolveDst, RenderTargetAccess, Range);
+					// RTV|ResolveDst is not a valid state for the platform RHI, use directly ERHIAccess::ResolveDst
+					AccessFunction(nullptr, ResolveTexture, ERHIAccess::ResolveDst, RenderTargetAccess, Range);
 				}
 			});
 
@@ -2677,6 +2676,95 @@ void FRDGBuilder::SubmitBufferUploads(FRHICommandListBase& RHICmdListUpload, UE:
 		AllocateUploadBuffersTask->Trigger();
 	}
 
+	if ( (RHICmdListUpload.NeedsExtraTransitions() && RHICmdListUpload.NeedTransitionsPatching()) && UploadedBuffers.Num() > 1)
+	{
+		SCOPED_NAMED_EVENT_TEXT("Upload_Multiple", FColor::Magenta);
+		
+		// This is here because we are explicitly batching a series of transition for all the buffers and we don't the individual extra transitions in Lock/Unlock
+		bool PrevCmdListAllowState = RHICmdListUpload.SetAllowExtraTransition(false);
+		
+		TSet<FRHIBuffer*, DefaultKeyFuncs<void*>, FConcurrentLinearSetAllocator> BuffersSet;
+		TArray<FRHITransitionInfo, FConcurrentLinearArrayAllocator> CopyDestTransitionInfo;
+		TArray<FRHITransitionInfo, FConcurrentLinearArrayAllocator> RevertTransitionInfo;
+		
+		BuffersSet.Reserve(UploadedBuffers.Num());
+		CopyDestTransitionInfo.Reserve(UploadedBuffers.Num());
+		RevertTransitionInfo.Reserve(UploadedBuffers.Num());
+		const EResourceLockMode LockMode = RLM_WriteOnly;
+		
+		// Lock all buffers, copy the data and create the transitions info
+		for (FUploadedBuffer& UploadedBuffer : UploadedBuffers)
+		{
+			FRDGBuffer* Buffer = UploadedBuffer.Buffer;
+
+			if (UploadedBuffer.DataFillCallback)
+			{
+				FRHIBuffer* RHIBuffer = Buffer->GetRHIUnchecked();
+				const uint32 DataSize = Buffer->Desc.GetSize();
+				void* DestPtr = RHICmdListUpload.LockBuffer(RHIBuffer, 0, DataSize, LockMode);
+				static_assert((LockMode == RLM_WriteOnly), "Transitions optimized only for RLM_WriteOnly");
+
+				UploadedBuffer.DataFillCallback(DestPtr, DataSize);
+
+				if (BuffersSet.Contains(RHIBuffer) == false)
+				{
+					BuffersSet.Add(RHIBuffer);
+					CopyDestTransitionInfo.Push(FRHITransitionInfo(RHIBuffer, ERHIAccess::Unknown, ERHIAccess::CopyDest, EResourceTransitionFlags::IgnoreAfterState));
+					RevertTransitionInfo.Push(FRHITransitionInfo(RHIBuffer, ERHIAccess::CopyDest, ERHIAccess::Unknown, EResourceTransitionFlags::IgnoreAfterState));
+				}
+			}
+			else
+			{
+				if (UploadedBuffer.bUseDataCallbacks)
+				{
+					UploadedBuffer.Data = UploadedBuffer.DataCallback();
+					UploadedBuffer.DataSize = UploadedBuffer.DataSizeCallback();
+				}
+
+				if (UploadedBuffer.Data && UploadedBuffer.DataSize)
+				{
+					FRHIBuffer* RHIBuffer = Buffer->GetRHIUnchecked();
+					check(UploadedBuffer.DataSize <= Buffer->Desc.GetSize());
+					void* DestPtr = RHICmdListUpload.LockBuffer(RHIBuffer, 0, UploadedBuffer.DataSize, LockMode);
+					static_assert((LockMode == RLM_WriteOnly), "Transitions optimized only for RLM_WriteOnly");
+
+					FMemory::Memcpy(DestPtr, UploadedBuffer.Data, UploadedBuffer.DataSize);
+
+					if (BuffersSet.Contains(RHIBuffer) == false)
+					{
+						BuffersSet.Add(RHIBuffer);
+						CopyDestTransitionInfo.Push(FRHITransitionInfo(RHIBuffer, ERHIAccess::Unknown, ERHIAccess::CopyDest, EResourceTransitionFlags::IgnoreAfterState));
+						RevertTransitionInfo.Push(FRHITransitionInfo(RHIBuffer, ERHIAccess::CopyDest, ERHIAccess::Unknown, EResourceTransitionFlags::IgnoreAfterState));
+					}
+				}
+			}
+		}
+
+		// Issue all COPY_DEST buffer transitions together
+		RHICmdListUpload.TransitionInternal(CopyDestTransitionInfo, ERHITransitionCreateFlags::AllowDuringRenderPass);
+
+		// Unlock all buffers
+		for (FUploadedBuffer& UploadedBuffer : UploadedBuffers)
+		{
+			FRDGBuffer* Buffer = UploadedBuffer.Buffer;
+
+			if (UploadedBuffer.DataFillCallback || (UploadedBuffer.Data && UploadedBuffer.DataSize))
+			{
+				RHICmdListUpload.UnlockBuffer(Buffer->GetRHIUnchecked());
+			}
+		}
+		
+		// Issue all Revert buffer transitions together
+		RHICmdListUpload.TransitionInternal(RevertTransitionInfo, ERHITransitionCreateFlags::AllowDuringRenderPass);
+
+		BuffersSet.Reset();
+		CopyDestTransitionInfo.Reset();
+		RevertTransitionInfo.Reset();
+
+		// Set the cmdlist state to the original state
+		RHICmdListUpload.SetAllowExtraTransition(PrevCmdListAllowState);
+	}
+	else
 	{
 		SCOPED_NAMED_EVENT_TEXT("Upload", FColor::Magenta);
 
