Modify D3D12 root signature setup to support work graph graphics nodes.
We need a new reserved binding space for the work graph mesh shader to separate it from the work graph pixel shader which uses the default binding space. This is required because work graphs don't support per stage visibility.
Added a default global root signature for work graphs that don't have any library objects that declare one (for example if they only include nodes with local root signatures).

diff --git a/Engine/Shaders/Private/ShaderBundleWorkGraphCommon.ush b/Engine/Shaders/Private/ShaderBundleWorkGraphCommon.ush
--- a/Engine/Shaders/Private/ShaderBundleWorkGraphCommon.ush
+++ b/Engine/Shaders/Private/ShaderBundleWorkGraphCommon.ush
@@ -2,7 +2,7 @@
 
 #include "Common.ush"
 
-#pragma message("UESHADERMETADATA_VERSION 7EDFC9CF-F736-47E6-84FE-D3B245272961")
+#pragma message("UESHADERMETADATA_VERSION DD3B6844-56DA-4913-B307-BE62A1F1882D")
 
 /** Record passed to shader bundle work graph nodes .*/
 struct FShaderBundleNodeRecord
diff --git a/Engine/Shaders/Shared/HLSLReservedSpaces.h b/Engine/Shaders/Shared/HLSLReservedSpaces.h
--- a/Engine/Shaders/Shared/HLSLReservedSpaces.h
+++ b/Engine/Shaders/Shared/HLSLReservedSpaces.h
@@ -16,6 +16,9 @@
 #define UE_HLSL_REGISTER(InType, InIndex, InSpace) register(UE_REG_CONCATENATE(InType, InIndex), UE_REG_CONCATENATE(space, InSpace))
 #endif
 
+// Default binding space
+#define UE_HLSL_SPACE_DEFAULT 0
+
 // Default register space for hit group (closest hit, any hit, intersection) shader resources
 #define UE_HLSL_SPACE_RAY_TRACING_LOCAL 0
 
@@ -25,12 +28,16 @@
 // Register space for "system" parameters (index buffer, vertex buffer, fetch parameters)
 #define UE_HLSL_SPACE_RAY_TRACING_SYSTEM 2
 
-// Default register space for local node shader resources
-#define UE_HLSL_SPACE_WORK_GRAPH_LOCAL 0
+// Register space for for resources referenced by a work graph pixel shader local root signature
+// Needs to match the default binding space since we don't specialize work graph pixel shader compilation
+#define UE_HLSL_SPACE_WORK_GRAPH_LOCAL_PIXEL UE_HLSL_SPACE_DEFAULT
 
-// Register space for global work graph shared shader resources
+// Register space for for resources referenced by a global work graph root signature
 #define UE_HLSL_SPACE_WORK_GRAPH_GLOBAL 1
 
+// Register space for resources referenced by a work graph compute or mesh shader local root signature
+#define UE_HLSL_SPACE_WORK_GRAPH_LOCAL 2
+
 // Register space for shader root constants (4xu32)
 #define UE_HLSL_SPACE_SHADER_ROOT_CONSTANTS 3
 
diff --git a/Engine/Source/Runtime/D3D12RHI/Private/D3D12Adapter.h b/Engine/Source/Runtime/D3D12RHI/Private/D3D12Adapter.h
--- a/Engine/Source/Runtime/D3D12RHI/Private/D3D12Adapter.h
+++ b/Engine/Source/Runtime/D3D12RHI/Private/D3D12Adapter.h
@@ -238,7 +238,10 @@ public:
 
 	const FD3D12RootSignature* GetRootSignature(const FBoundShaderStateInput& BoundShaderState);
 	const FD3D12RootSignature* GetRootSignature(const class FD3D12ComputeShader* Shader);
+	
 	const FD3D12RootSignature* GetRootSignature(const class FD3D12WorkGraphShader* Shader);
+	const FD3D12RootSignature* GetGlobalWorkGraphRootSignature(const FRHIShaderBindingLayout& ShaderBindingLayout);
+	const FD3D12RootSignature* GetWorkGraphGraphicsRootSignature(const FBoundShaderStateInput& BoundShaderState);
 
 	const FD3D12RootSignature* GetLocalRootSignature(const class FD3D12RayTracingShader* Shader);
 	const FD3D12RootSignature* GetGlobalRayTracingRootSignature(const FRHIShaderBindingLayout& ShaderBindingLayout);
diff --git a/Engine/Source/Runtime/D3D12RHI/Private/D3D12RayTracing.cpp b/Engine/Source/Runtime/D3D12RHI/Private/D3D12RayTracing.cpp
--- a/Engine/Source/Runtime/D3D12RHI/Private/D3D12RayTracing.cpp
+++ b/Engine/Source/Runtime/D3D12RHI/Private/D3D12RayTracing.cpp
@@ -896,11 +896,11 @@ public:
 		}
 		else
 		{
-		LocalRootSignatureDesc.Version = D3D_ROOT_SIGNATURE_VERSION_1_0;
-		LocalRootSignatureDesc.Desc_1_0.Flags |= D3D12_ROOT_SIGNATURE_FLAG_LOCAL_ROOT_SIGNATURE;
+			LocalRootSignatureDesc.Version = D3D_ROOT_SIGNATURE_VERSION_1_0;
+			LocalRootSignatureDesc.Desc_1_0.Flags |= D3D12_ROOT_SIGNATURE_FLAG_LOCAL_ROOT_SIGNATURE;
 		}
 		
-		DefaultLocalRootSignature.Init(LocalRootSignatureDesc);
+		DefaultLocalRootSignature.Init(LocalRootSignatureDesc, RS_RayTracingLocal);
 	}
 
 	~FD3D12RayTracingPipelineCache()
diff --git a/Engine/Source/Runtime/D3D12RHI/Private/D3D12RootSignature.cpp b/Engine/Source/Runtime/D3D12RHI/Private/D3D12RootSignature.cpp
--- a/Engine/Source/Runtime/D3D12RHI/Private/D3D12RootSignature.cpp
+++ b/Engine/Source/Runtime/D3D12RHI/Private/D3D12RootSignature.cpp
@@ -111,6 +111,53 @@ FORCEINLINE D3D12_ROOT_SIGNATURE_FLAGS GetD3D12RootSignatureDenyFlag(EShaderVisi
 	};
 }
 
+static bool IsLocalRootSignature(ERootSignatureType InRootSignatureType)
+{
+	switch (InRootSignatureType)
+	{
+	case RS_RayTracingLocal:
+	case RS_WorkGraphLocalCompute:
+	case RS_WorkGraphLocalRaster:
+		return true;
+	default:
+		return false;
+	}
+}
+
+static uint32 GetBindingSpace(ERootSignatureType InRootSignatureType, bool bInIsPixelShader)
+{
+	switch (InRootSignatureType)
+	{
+	case RS_RayTracingGlobal:
+		return UE_HLSL_SPACE_RAY_TRACING_GLOBAL;
+	case RS_RayTracingLocal:
+		return UE_HLSL_SPACE_RAY_TRACING_LOCAL;
+	case RS_WorkGraphGlobal:
+		return UE_HLSL_SPACE_WORK_GRAPH_GLOBAL;
+	case RS_WorkGraphLocalCompute:
+		return UE_HLSL_SPACE_WORK_GRAPH_LOCAL;
+	case RS_WorkGraphLocalRaster:
+		// We use separate binding spaces for the mesh and pixel stages.
+		// If we add support for work graph vertex shaders then we will need to add a new binding space for the additional stage.
+		return bInIsPixelShader ? UE_HLSL_SPACE_WORK_GRAPH_LOCAL_PIXEL : UE_HLSL_SPACE_WORK_GRAPH_LOCAL;
+	default:
+		break;
+	}
+
+	// Default binding space for D3D 11 & 12 shaders
+	return UE_HLSL_SPACE_DEFAULT;
+}
+
+static uint32 GetBindingSpace(ERootSignatureType InRootSignatureType, EShaderVisibility InVisibility)
+{
+	return GetBindingSpace(InRootSignatureType, InVisibility == SV_Pixel);
+}
+
+static uint32 GetBindingSpace(ERootSignatureType InRootSignatureType, EShaderFrequency InFrequency)
+{
+	return GetBindingSpace(InRootSignatureType, InFrequency == SF_Pixel);
+}
+
 
 FD3D12RootSignatureDesc::FD3D12RootSignatureDesc(const FD3D12QuantizedBoundShaderState& QBSS, const D3D12_RESOURCE_BINDING_TIER ResourceBindingTier)
 	: RootParametersSize(0)
@@ -149,27 +196,13 @@ FD3D12RootSignatureDesc::FD3D12RootSignatureDesc(const FD3D12QuantizedBoundShade
 
 	const D3D12_ROOT_DESCRIPTOR_FLAGS CBVRootDescriptorFlags = D3D12_ROOT_DESCRIPTOR_FLAG_DATA_STATIC;	// We always set the data in an upload heap before calling Set*RootConstantBufferView.
 
-	uint32 BindingSpace = 0; // Default binding space for D3D 11 & 12 shaders
-
 	const bool bUseShaderDiagnosticBuffer = D3D12_ALLOW_SHADER_DIAGNOSTIC_BUFFER
 		&& QBSS.bUseDiagnosticBuffer
-		&& QBSS.RootSignatureType != RS_RayTracingLocal
-		&& QBSS.RootSignatureType != RS_WorkGraphLocal;
-
-	if (QBSS.RootSignatureType == RS_WorkGraphLocal)
-	{
-		BindingSpace = UE_HLSL_SPACE_WORK_GRAPH_LOCAL;
-	}
-	else if (QBSS.RootSignatureType == RS_WorkGraphGlobal)
-	{
-		BindingSpace = UE_HLSL_SPACE_WORK_GRAPH_GLOBAL;
-	}
+		&& !IsLocalRootSignature(QBSS.RootSignatureType);
 
 #if D3D12_RHI_RAYTRACING
 	if (QBSS.RootSignatureType == RS_RayTracingLocal)
 	{
-		BindingSpace = UE_HLSL_SPACE_RAY_TRACING_LOCAL;
-
 		// Add standard root parameters for hit groups, as per FD3D12HitGroupSystemParameters declaration in D3D12Util.h 
 		// and RayTracingHitGroupCommon.ush (Non-Bindless) and D3DCommon.ush (Bindless):
 		//          FHitGroupSystemRootConstants:
@@ -231,13 +264,9 @@ FD3D12RootSignatureDesc::FD3D12RootSignatureDesc(const FD3D12QuantizedBoundShade
 			}
 		}	
 	}
-	else if (QBSS.RootSignatureType == RS_RayTracingGlobal)
-	{
-		BindingSpace = UE_HLSL_SPACE_RAY_TRACING_GLOBAL;
-	}
 #endif //D3D12_RHI_RAYTRACING
 
-	const uint32 RootDescriptorTableCost = QBSS.RootSignatureType == RS_RayTracingLocal || QBSS.RootSignatureType == RS_WorkGraphLocal ? RootDescriptorTableCostLocal : RootDescriptorTableCostGlobal;
+	const uint32 RootDescriptorTableCost = IsLocalRootSignature(QBSS.RootSignatureType) ? RootDescriptorTableCostLocal : RootDescriptorTableCostGlobal;
 
 	// For each root parameter type...
 	for (uint32 RootParameterTypeIndex = 0; RootParameterTypeIndex < UE_ARRAY_COUNT(RootParameterTypePriorityOrder); RootParameterTypeIndex++)
@@ -247,8 +276,11 @@ FD3D12RootSignatureDesc::FD3D12RootSignatureDesc(const FD3D12QuantizedBoundShade
 		// ... and each shader stage visibility ...
 		for (uint32 ShaderVisibilityIndex = 0; ShaderVisibilityIndex < UE_ARRAY_COUNT(ShaderVisibilityPriorityOrder); ShaderVisibilityIndex++)
 		{
-			const EShaderVisibility Visibility = ShaderVisibilityPriorityOrder[ShaderVisibilityIndex];
-			const FShaderRegisterCounts& Shader = QBSS.RegisterCounts[Visibility];
+			const EShaderVisibility StageVisibility = ShaderVisibilityPriorityOrder[ShaderVisibilityIndex];
+			const FShaderRegisterCounts& Shader = QBSS.RegisterCounts[StageVisibility];
+			// Work graphs must use SV_All and we use the per frequency binding space instead.
+			const EShaderVisibility Visibility = QBSS.RootSignatureType == RS_WorkGraphLocalRaster ? SV_All : StageVisibility;
+			const uint32 BindingSpace = GetBindingSpace(QBSS.RootSignatureType, StageVisibility);
 
 			switch (RootParameterType)
 			{
@@ -265,7 +297,7 @@ FD3D12RootSignatureDesc::FD3D12RootSignatureDesc(const FD3D12QuantizedBoundShade
 
 				if (Shader.ConstantBufferCount > MAX_ROOT_CBVS)
 				{
-					checkf(QBSS.RootSignatureType != RS_RayTracingLocal && QBSS.RootSignatureType != RS_WorkGraphLocal, TEXT("CBV descriptor tables are not implemented for local root signatures"));
+					checkf(!IsLocalRootSignature(QBSS.RootSignatureType), TEXT("CBV descriptor tables are not implemented for local root signatures"));
 
 					// Use a descriptor table for the 'excess' CBVs
 					check(RootParameterCount < MaxRootParameters);
@@ -327,11 +359,11 @@ FD3D12RootSignatureDesc::FD3D12RootSignatureDesc(const FD3D12QuantizedBoundShade
 	}
 
 #if D3D12_RHI_RAYTRACING
-	if (QBSS.RootSignatureType == RS_RayTracingLocal || QBSS.RootSignatureType == RS_WorkGraphLocal)
+	if (IsLocalRootSignature(QBSS.RootSignatureType))
 	{
 		Flags = D3D12_ROOT_SIGNATURE_FLAG_LOCAL_ROOT_SIGNATURE;
 	}
-	else if(QBSS.RootSignatureType == RS_RayTracingGlobal || QBSS.RootSignatureType == RS_WorkGraphGlobal)
+	else if(QBSS.RootSignatureType == RS_RayTracingGlobal)
 	{
 		Flags |= FD3D12_ROOT_SIGNATURE_FLAG_GLOBAL_ROOT_SIGNATURE;
 	}
@@ -434,7 +466,7 @@ FD3D12RootSignatureDesc::FD3D12RootSignatureDesc(const FD3D12QuantizedBoundShade
 #endif
 
 #if D3D12_RHI_RAYTRACING
-	if (QBSS.RootSignatureType == RS_RayTracingLocal || QBSS.RootSignatureType == RS_WorkGraphLocal)
+	if (IsLocalRootSignature(QBSS.RootSignatureType))
 	{
 		// Local root signatures don't need to provide static samplers as they are provided by global RS already.
 		// Providing static sampler bindings in global and local RS simultaneously is invalid due to overlapping register ranges.
@@ -454,14 +486,13 @@ FD3D12RootSignatureDesc::FD3D12RootSignatureDesc(const FD3D12QuantizedBoundShade
 			RootDesc.Init_1_1(RootParameterCount, TableSlots, 0, nullptr, Flags);
 		}
 	}
-
 }
 
 void FD3D12RootSignature::InitStaticGraphicsRootSignature(EShaderBindingLayoutFlags InFlags)
 {
 	D3D12ShaderUtils::FBinaryRootSignatureCreator Creator;
 	D3D12ShaderUtils::CreateGfxRootSignature(Creator, InFlags);
-	Init(Creator.Finalize());
+	Init(Creator.Finalize(), RS_Raster);
 
 	for (int32 ParameterSlot = 0; ParameterSlot < Creator.Parameters.Num(); ++ParameterSlot)
 	{
@@ -487,7 +518,7 @@ void FD3D12RootSignature::InitStaticComputeRootSignatureDesc(EShaderBindingLayou
 {
 	D3D12ShaderUtils::FBinaryRootSignatureCreator Creator;
 	D3D12ShaderUtils::CreateComputeRootSignature(Creator, InFlags);
-	Init(Creator.Finalize());
+	Init(Creator.Finalize(), RS_Raster);
 
 	for (int32 ParameterSlot = 0; ParameterSlot < Creator.Parameters.Num(); ++ParameterSlot)
 	{
@@ -514,14 +545,14 @@ void FD3D12RootSignature::InitStaticRayTracingGlobalRootSignatureDesc(EShaderBin
 {
 	D3D12ShaderUtils::FBinaryRootSignatureCreator Creator;
 	D3D12ShaderUtils::CreateRayTracingSignature(Creator, false, FD3D12_ROOT_SIGNATURE_FLAG_GLOBAL_ROOT_SIGNATURE, InFlags);
-	Init(Creator.Finalize(), UE_HLSL_SPACE_RAY_TRACING_GLOBAL);
+	Init(Creator.Finalize(), RS_RayTracingGlobal);
 }
 
 void FD3D12RootSignature::InitStaticRayTracingLocalRootSignatureDesc(EShaderBindingLayoutFlags InFlags)
 {
 	D3D12ShaderUtils::FBinaryRootSignatureCreator Creator;
 	D3D12ShaderUtils::CreateRayTracingSignature(Creator, true, D3D12_ROOT_SIGNATURE_FLAG_LOCAL_ROOT_SIGNATURE, InFlags);
-	Init(Creator.Finalize(), UE_HLSL_SPACE_RAY_TRACING_LOCAL);
+	Init(Creator.Finalize(), RS_RayTracingLocal);
 }
 #endif // D3D12_RHI_RAYTRACING
 
@@ -537,29 +568,10 @@ void FD3D12RootSignature::Init(const FD3D12QuantizedBoundShaderState& InQBSS)
 	StaticShaderBindingCount = Desc.GetStaticShaderBindingCount();
 	DiagnosticBufferSlot = Desc.GetDiagnosticBufferSlot();
 
-	uint32 BindingSpace = 0; // Default binding space for D3D 11 & 12 shaders
-
-	if (InQBSS.RootSignatureType == RS_WorkGraphLocal)
-	{
-		BindingSpace = UE_HLSL_SPACE_WORK_GRAPH_LOCAL;
-	}
-	else if (InQBSS.RootSignatureType == RS_WorkGraphGlobal)
-	{
-		BindingSpace = UE_HLSL_SPACE_WORK_GRAPH_GLOBAL;
-	}
-	else if (InQBSS.RootSignatureType == RS_RayTracingGlobal)
-	{
-		BindingSpace = UE_HLSL_SPACE_RAY_TRACING_GLOBAL;
-	}
-	else if (InQBSS.RootSignatureType == RS_RayTracingLocal)
-	{
-		BindingSpace = UE_HLSL_SPACE_RAY_TRACING_LOCAL;
-	}
-
-	Init(Desc.GetDesc(), BindingSpace);
+	Init(Desc.GetDesc(), InQBSS.RootSignatureType);
 }
 
-void FD3D12RootSignature::Init(const D3D12_VERSIONED_ROOT_SIGNATURE_DESC& InDesc, uint32 BindingSpace)
+void FD3D12RootSignature::Init(const D3D12_VERSIONED_ROOT_SIGNATURE_DESC& InDesc, ERootSignatureType InRootSignatureType)
 {
 	ID3D12Device* Device = GetParentAdapter()->GetD3DDevice();
 	
@@ -579,20 +591,20 @@ void FD3D12RootSignature::Init(const D3D12_VERSIONED_ROOT_SIGNATURE_DESC& InDesc
 		RootSignatureBlob->GetBufferSize(),
 		IID_PPV_ARGS(RootSignature.GetInitReference())));
 
-	AnalyzeSignature(InDesc, BindingSpace);
+	AnalyzeSignature(InDesc, InRootSignatureType);
 	// TODO: Analyze vendor extension space?
 }
 
-void FD3D12RootSignature::AnalyzeSignature(const D3D12_VERSIONED_ROOT_SIGNATURE_DESC& Desc, uint32 BindingSpace)
+void FD3D12RootSignature::AnalyzeSignature(const D3D12_VERSIONED_ROOT_SIGNATURE_DESC& Desc, ERootSignatureType InRootSignatureType)
 {
 	switch (Desc.Version)
 	{
 	case D3D_ROOT_SIGNATURE_VERSION_1_0:
-		InternalAnalyzeSignature(Desc.Desc_1_0, BindingSpace);
+		InternalAnalyzeSignature(Desc.Desc_1_0, InRootSignatureType);
 		break;
 
 	case D3D_ROOT_SIGNATURE_VERSION_1_1:
-		InternalAnalyzeSignature(Desc.Desc_1_1, BindingSpace);
+		InternalAnalyzeSignature(Desc.Desc_1_1, InRootSignatureType);
 		break;
 
 	default:
@@ -602,7 +614,7 @@ void FD3D12RootSignature::AnalyzeSignature(const D3D12_VERSIONED_ROOT_SIGNATURE_
 }
 
 template<typename RootSignatureDescType>
-void FD3D12RootSignature::InternalAnalyzeSignature(const RootSignatureDescType& Desc, uint32 BindingSpace)
+void FD3D12RootSignature::InternalAnalyzeSignature(const RootSignatureDescType& Desc, ERootSignatureType InRootSignatureType)
 {
 	// Reset members to default values.
 	{
@@ -666,17 +678,23 @@ void FD3D12RootSignature::InternalAnalyzeSignature(const RootSignatureDescType&
 			break;
 		}
 
-		if (ParameterBindingSpace != BindingSpace)
-		{
-			// Only consider parameters in the requested binding space.
-			continue;
-		}
-
 		EShaderFrequency CurrentVisibleSF = SF_NumFrequencies;
 		switch (CurrentParameter.ShaderVisibility)
 		{
 		case D3D12_SHADER_VISIBILITY_ALL:
-			CurrentVisibleSF = SF_NumFrequencies;
+			if (InRootSignatureType == RS_WorkGraphLocalRaster)
+			{
+				// Handle work graph graphics node case where shader visibility will be SV_All.
+				// We need to use the stored parameter binding space to determine which stage this is.
+				if (ParameterBindingSpace == UE_HLSL_SPACE_WORK_GRAPH_LOCAL)
+				{
+					CurrentVisibleSF = SF_Mesh;
+				}
+				else if (ParameterBindingSpace == UE_HLSL_SPACE_WORK_GRAPH_LOCAL_PIXEL)
+				{
+					CurrentVisibleSF = SF_Pixel;
+				}
+			}
 			break;
 
 		case D3D12_SHADER_VISIBILITY_VERTEX:
@@ -703,6 +721,13 @@ void FD3D12RootSignature::InternalAnalyzeSignature(const RootSignatureDescType&
 			break;
 		}
 
+		const uint32 BindingSpace = GetBindingSpace(InRootSignatureType, CurrentVisibleSF);
+		if (BindingSpace != ParameterBindingSpace)
+		{
+			// Only consider parameters in the requested binding space.
+			continue;
+		}
+
 		// Determine shader stage visibility.
 		{
 			Stage[SF_Vertex].bVisible = Stage[SF_Vertex].bVisible || (!bDenyVS && HasVisibility(CurrentParameter.ShaderVisibility, D3D12_SHADER_VISIBILITY_VERTEX));
diff --git a/Engine/Source/Runtime/D3D12RHI/Private/D3D12RootSignature.h b/Engine/Source/Runtime/D3D12RHI/Private/D3D12RootSignature.h
--- a/Engine/Source/Runtime/D3D12RHI/Private/D3D12RootSignature.h
+++ b/Engine/Source/Runtime/D3D12RHI/Private/D3D12RootSignature.h
@@ -95,7 +95,7 @@ public:
 	}
 
 	void Init(const FD3D12QuantizedBoundShaderState& InQBSS);
-	void Init(const D3D12_VERSIONED_ROOT_SIGNATURE_DESC& InDesc, uint32 BindingSpace = 0);
+	void Init(const D3D12_VERSIONED_ROOT_SIGNATURE_DESC& InDesc, ERootSignatureType InRootSignatureType);
 
 	void InitStaticGraphicsRootSignature(EShaderBindingLayoutFlags InFlags);
 	void InitStaticComputeRootSignatureDesc(EShaderBindingLayoutFlags InFlags);
@@ -224,10 +224,10 @@ public:
 	inline int8 GetDiagnosticBufferSlot() const { return DiagnosticBufferSlot; }
 
 private:
-	void AnalyzeSignature(const D3D12_VERSIONED_ROOT_SIGNATURE_DESC& Desc, uint32 BindingSpace);
+	void AnalyzeSignature(const D3D12_VERSIONED_ROOT_SIGNATURE_DESC& Desc, ERootSignatureType InRootSignatureType);
 
 	template<typename RootSignatureDescType>
-	void InternalAnalyzeSignature(const RootSignatureDescType& Desc, uint32 BindingSpace);
+	void InternalAnalyzeSignature(const RootSignatureDescType& Desc, ERootSignatureType InRootSignatureType);
 
 	inline bool HasVisibility(const D3D12_SHADER_VISIBILITY& ParameterVisibility, const D3D12_SHADER_VISIBILITY& Visibility) const
 	{
diff --git a/Engine/Source/Runtime/D3D12RHI/Private/D3D12Util.cpp b/Engine/Source/Runtime/D3D12RHI/Private/D3D12Util.cpp
--- a/Engine/Source/Runtime/D3D12RHI/Private/D3D12Util.cpp
+++ b/Engine/Source/Runtime/D3D12RHI/Private/D3D12Util.cpp
@@ -1357,12 +1357,50 @@ const FD3D12RootSignature* FD3D12Adapter::GetRootSignature(const FD3D12WorkGraph
 	FD3D12QuantizedBoundShaderState QBSS{};
 	QuantizeBoundShaderStateCommon(QBSS, WorkGraphShader, GetResourceBindingTier(), SV_All, true /*bAllowUAVs*/);
 
-	QBSS.RootSignatureType = WorkGraphShader->GetFrequency() == SF_WorkGraphRoot ? RS_WorkGraphGlobal : RS_WorkGraphLocal;
+	QBSS.RootSignatureType = WorkGraphShader->GetFrequency() == SF_WorkGraphRoot ? RS_WorkGraphGlobal : RS_WorkGraphLocalCompute;
 	check(QBSS.bAllowIAInputLayout == false); // No access to vertex buffers needed
 
 	return RootSignatureManager.GetRootSignature(QBSS);
 }
 
+const FD3D12RootSignature* FD3D12Adapter::GetGlobalWorkGraphRootSignature(const FRHIShaderBindingLayout& ShaderBindingLayout)
+{
+	FD3D12QuantizedBoundShaderState QBSS{};
+	FShaderRegisterCounts& QBSSRegisterCounts = QBSS.RegisterCounts[SV_All];
+
+	QBSS.ShaderBindingLayout = ShaderBindingLayout;
+	QBSS.RootSignatureType = RS_WorkGraphGlobal;
+	QBSS.bUseDiagnosticBuffer = true;
+
+#if PLATFORM_SUPPORTS_BINDLESS_RENDERING
+	QBSS.bUseDirectlyIndexedResourceHeap = GetBindlessDescriptorAllocator().AreResourcesBindless();
+	QBSS.bUseDirectlyIndexedSamplerHeap = GetBindlessDescriptorAllocator().AreSamplersBindless();
+#endif
+
+	QBSSRegisterCounts.SamplerCount = MAX_SAMPLERS;
+	QBSSRegisterCounts.ShaderResourceCount = MAX_SRVS;
+	QBSSRegisterCounts.ConstantBufferCount = MAX_CBS;
+	QBSSRegisterCounts.UnorderedAccessCount = MAX_UAVS;
+
+	return RootSignatureManager.GetRootSignature(QBSS);
+}
+
+const FD3D12RootSignature* FD3D12Adapter::GetWorkGraphGraphicsRootSignature(const FBoundShaderStateInput& BSS)
+{
+	FD3D12QuantizedBoundShaderState QBSS{};
+
+	const D3D12_RESOURCE_BINDING_TIER ResourceBindingTier = GetResourceBindingTier();
+
+#if PLATFORM_SUPPORTS_MESH_SHADERS
+	QuantizeBoundShaderStateCommon(QBSS, FD3D12DynamicRHI::ResourceCast(BSS.GetMeshShader()), ResourceBindingTier, SV_Mesh);
+#endif
+	QuantizeBoundShaderStateCommon(QBSS, FD3D12DynamicRHI::ResourceCast(BSS.GetPixelShader()), ResourceBindingTier, SV_Pixel, true /*bAllowUAVs*/);
+
+	QBSS.RootSignatureType = RS_WorkGraphLocalRaster;
+
+	return RootSignatureManager.GetRootSignature(QBSS);
+}
+
 #if D3D12_RHI_RAYTRACING
 
 const FD3D12RootSignature* FD3D12Adapter::GetGlobalRayTracingRootSignature(const FRHIShaderBindingLayout& ShaderBindingLayout)
diff --git a/Engine/Source/Runtime/D3D12RHI/Private/D3D12Util.h b/Engine/Source/Runtime/D3D12RHI/Private/D3D12Util.h
--- a/Engine/Source/Runtime/D3D12RHI/Private/D3D12Util.h
+++ b/Engine/Source/Runtime/D3D12RHI/Private/D3D12Util.h
@@ -68,13 +68,14 @@ enum EShaderVisibility
 	SV_ShaderVisibilityCount
 };
 
-enum ERTRootSignatureType
+enum ERootSignatureType
 {
 	RS_Raster,
 	RS_RayTracingGlobal,
 	RS_RayTracingLocal,
 	RS_WorkGraphGlobal,
-	RS_WorkGraphLocal,
+	RS_WorkGraphLocalCompute,
+	RS_WorkGraphLocalRaster,
 };
 
 struct FShaderRegisterCounts
@@ -106,7 +107,7 @@ struct FD3D12QuantizedBoundShaderState
 {
 	FRHIShaderBindingLayout ShaderBindingLayout;
 	TStaticArray<FShaderRegisterCounts, SV_ShaderVisibilityCount> RegisterCounts;
-	ERTRootSignatureType RootSignatureType = RS_Raster;
+	ERootSignatureType RootSignatureType = RS_Raster;
 	union
 	{
 		struct
