Update FWorkGraphPipelineStateInitializer to not be so tightly coupled to shader bundle usage.
We can now use any combination of workgraph shader nodes, with either 0 or 1 using global root signature.
Updated work graph shader bundles to use entry node with local root signature.

diff --git a/Engine/Shaders/Private/ShaderBundleWorkGraphCommon.ush b/Engine/Shaders/Private/ShaderBundleWorkGraphCommon.ush
--- a/Engine/Shaders/Private/ShaderBundleWorkGraphCommon.ush
+++ b/Engine/Shaders/Private/ShaderBundleWorkGraphCommon.ush
@@ -2,7 +2,7 @@
 
 #include "Common.ush"
 
-#pragma message("UESHADERMETADATA_VERSION DD3B6844-56DA-4913-B307-BE62A1F1882D")
+#pragma message("UESHADERMETADATA_VERSION 3EA3DFAD-BDA9-4554-B218-E9F05711625D")
 
 /** Record passed to shader bundle work graph nodes .*/
 struct FShaderBundleNodeRecord
diff --git a/Engine/Source/Runtime/D3D12RHI/Private/D3D12WorkGraph.cpp b/Engine/Source/Runtime/D3D12RHI/Private/D3D12WorkGraph.cpp
--- a/Engine/Source/Runtime/D3D12RHI/Private/D3D12WorkGraph.cpp
+++ b/Engine/Source/Runtime/D3D12RHI/Private/D3D12WorkGraph.cpp
@@ -24,78 +24,97 @@ static FAutoConsoleVariableRef CVarShaderBundleSkipDispatch(
 FD3D12WorkGraphPipelineState::FD3D12WorkGraphPipelineState(FD3D12Device* Device, const FWorkGraphPipelineStateInitializer& Initializer)
 {
 #if D3D12_RHI_WORKGRAPHS
+	
+	RootArgStrideInBytes = 0;
 
-	ID3D12Device9* Device9 = (ID3D12Device9*)Device->GetDevice();
-
-	Shader = (FD3D12WorkGraphShader*)Initializer.GetShader();
-	TCHAR const* ProgramName = Initializer.GetProgramName().IsEmpty() ? TEXT("WorkGraphProgram") : *Initializer.GetProgramName();
+	// Use global root signature from specified node, or if non-specified use a fixed root signature.
+	ID3D12RootSignature* GlobalRootSignature = nullptr;
+	const int32 RootShaderIndex = Initializer.GetRootShaderIndex();
+	if (RootShaderIndex != -1)
+	{
+		if (Initializer.GetShaderTable().IsValidIndex(RootShaderIndex))
+		{
+			FRHIWorkGraphShader* Shader = Initializer.GetShaderTable()[RootShaderIndex];
+			if (Shader != nullptr && Shader->GetFrequency() == SF_WorkGraphRoot)
+			{
+				GlobalRootSignature = ((FD3D12WorkGraphShader*)Shader)->RootSignature->GetRootSignature();
+			}
+		}
+	}
+	if (GlobalRootSignature == nullptr)
+	{
+		FD3D12Adapter* Adapter = Device->GetParentAdapter();
+		FRHIShaderBindingLayout ShaderBindingLayout; // todo pass ShaderBindingLayout down
+		GlobalRootSignature = Adapter->GetGlobalWorkGraphRootSignature(ShaderBindingLayout)->GetRootSignature();
+	}
+	RootSignature = GlobalRootSignature;
 
 	CD3DX12_STATE_OBJECT_DESC StateObjectDesc(D3D12_STATE_OBJECT_TYPE_EXECUTABLE);
+	
+	CD3DX12_GLOBAL_ROOT_SIGNATURE_SUBOBJECT* GlobalRootSignatureSubobject = StateObjectDesc.CreateSubobject<CD3DX12_GLOBAL_ROOT_SIGNATURE_SUBOBJECT>();
+	GlobalRootSignatureSubobject->SetRootSignature(RootSignature);
 
-	CD3DX12_GLOBAL_ROOT_SIGNATURE_SUBOBJECT* GlobalRootSignature = StateObjectDesc.CreateSubobject<CD3DX12_GLOBAL_ROOT_SIGNATURE_SUBOBJECT>();
-	GlobalRootSignature->SetRootSignature(Shader->RootSignature->GetRootSignature());
+	CD3DX12_WORK_GRAPH_SUBOBJECT* WorkGraphSubobject = StateObjectDesc.CreateSubobject<CD3DX12_WORK_GRAPH_SUBOBJECT>();
+	TCHAR const* ProgramName = Initializer.GetProgramName().IsEmpty() ? TEXT("WorkGraphProgram") : *Initializer.GetProgramName();
+	WorkGraphSubobject->SetProgramName(ProgramName);
 
+	D3D12_NODE_ID EntryPoint;
 	{
-		CD3DX12_DXIL_LIBRARY_SUBOBJECT* Lib = StateObjectDesc.CreateSubobject<CD3DX12_DXIL_LIBRARY_SUBOBJECT>();
-		CD3DX12_SHADER_BYTECODE LibCode(Shader->Code.GetData(), Shader->Code.Num());
-		Lib->SetDXILLibrary(&LibCode);
+		FD3D12WorkGraphShader* EntryShader = (FD3D12WorkGraphShader*)Initializer.GetShaderTable()[0];
+		EntryPoint.Name = *EntryShader->EntryPoint;
+		EntryPoint.ArrayIndex = 0;
 	}
+	WorkGraphSubobject->AddEntrypoint(EntryPoint);
 
-	const int32 ShaderBundleNum = Initializer.GetShaderBundleNodeTable().Num();
-	const bool bIsShaderBundle = ShaderBundleNum > 0;
-
-	RootArgStrideInBytes = 0;
-
-	if (bIsShaderBundle)
+	// Compute Shader Table
+	const int32 ShaderTableNum = Initializer.GetShaderTable().Num();
+	for (int32 Index = 0; Index < ShaderTableNum; ++Index) 
 	{
-		for (int32 Index = 0; Index < ShaderBundleNum; ++Index)
+		FD3D12WorkGraphShader* NodeShader = (FD3D12WorkGraphShader*)Initializer.GetShaderTable()[Index];
+
+		if (NodeShader != nullptr)
 		{
-			FD3D12WorkGraphShader* NodeShader = (FD3D12WorkGraphShader*)Initializer.GetShaderBundleNodeTable()[Index];
-			LocalNodeShaders.Add(NodeShader);
-			if (NodeShader != nullptr)
+			const int32 NodeNameIndex = Initializer.GetNameTable().IndexOfByPredicate([NodeShader](FWorkGraphPipelineStateInitializer::FNameMap const& NameMap)
 			{
+				return NameMap.ExportName == NodeShader->EntryPoint;
+			});
+
+			if (NodeNameIndex != INDEX_NONE)
+			{
+				FString const& NodeName = Initializer.GetNameTable()[NodeNameIndex].NodeName;
+				const uint32 NodeCount = (NodeCountPerName.FindOrAdd(NodeName))++;
+				const uint32 NodeArrayIndex = NodeCount;
+				const FString NodePathName = FString::Printf(TEXT("%s_%d"), *NodeName, NodeArrayIndex);
+
 				CD3DX12_DXIL_LIBRARY_SUBOBJECT* Lib = StateObjectDesc.CreateSubobject<CD3DX12_DXIL_LIBRARY_SUBOBJECT>();
 				CD3DX12_SHADER_BYTECODE LibCode(NodeShader->Code.GetData(), NodeShader->Code.Num());
 				Lib->SetDXILLibrary(&LibCode);
-
-				FString NodeName = FString::Printf(TEXT("%s_%d"), *Initializer.GetShaderBundleNodeName(), Index);
-				Lib->DefineExport(*NodeName, *NodeShader->EntryPoint);
+				Lib->DefineExport(*NodePathName, *NodeShader->EntryPoint);
 
 				CD3DX12_LOCAL_ROOT_SIGNATURE_SUBOBJECT* LocalRootSignature = StateObjectDesc.CreateSubobject<CD3DX12_LOCAL_ROOT_SIGNATURE_SUBOBJECT>();
 				LocalRootSignature->SetRootSignature(NodeShader->RootSignature->GetRootSignature());
-				CD3DX12_SUBOBJECT_TO_EXPORTS_ASSOCIATION_SUBOBJECT* Association = StateObjectDesc.CreateSubobject<CD3DX12_SUBOBJECT_TO_EXPORTS_ASSOCIATION_SUBOBJECT>();
-				Association->SetSubobjectToAssociate(*LocalRootSignature);
-				Association->AddExport(*NodeName);
-
-				RootArgStrideInBytes = FMath::Max(RootArgStrideInBytes, NodeShader->RootSignature->GetTotalRootSignatureSizeInBytes());
+				CD3DX12_SUBOBJECT_TO_EXPORTS_ASSOCIATION_SUBOBJECT* AssociationSubobject = StateObjectDesc.CreateSubobject<CD3DX12_SUBOBJECT_TO_EXPORTS_ASSOCIATION_SUBOBJECT>();
+				AssociationSubobject->SetSubobjectToAssociate(*LocalRootSignature);
+				AssociationSubobject->AddExport(*NodePathName);
+
+				CD3DX12_COMMON_COMPUTE_NODE_OVERRIDES* Override = WorkGraphSubobject->CreateCommonComputeNodeOverrides(*NodePathName);
+				Override->NewName(D3D12_NODE_ID{ *NodeName, NodeArrayIndex });
+				
+				const int32 LocalRootArgumentsTableIndex = RootArgOffsets.Num();
+				RootArgOffsets.Add(LocalRootArgumentsTableIndex);
+				Override->LocalRootArgumentsTableIndex(LocalRootArgumentsTableIndex);
 			}
-		}
-		RootArgStrideInBytes = ((RootArgStrideInBytes + 15) & ~15);
-	}
-
-	CD3DX12_WORK_GRAPH_SUBOBJECT* WorkGraph = StateObjectDesc.CreateSubobject<CD3DX12_WORK_GRAPH_SUBOBJECT>();
-	WorkGraph->SetProgramName(ProgramName);
-	
-	D3D12_NODE_ID EntryPoint;
-	EntryPoint.Name = *Shader->EntryPoint;
-	EntryPoint.ArrayIndex = 0;
-	WorkGraph->AddEntrypoint(EntryPoint);
 
-	if (bIsShaderBundle)
-	{
-		for (int32 Index = 0; Index < ShaderBundleNum; ++Index)
-		{
-			if (LocalNodeShaders[Index])
-			{
-				FString NodeName = FString::Printf(TEXT("%s_%d"), *Initializer.GetShaderBundleNodeName(), Index);
-				CD3DX12_COMMON_COMPUTE_NODE_OVERRIDES* Override = WorkGraph->CreateCommonComputeNodeOverrides(*NodeName);
-				Override->NewName(D3D12_NODE_ID{ *Initializer.GetShaderBundleNodeName(), (uint32)Index });
-			}
+			RootArgStrideInBytes = FMath::Max(RootArgStrideInBytes, NodeShader->RootSignature->GetTotalRootSignatureSizeInBytes());
 		}
 	}
+	
+	RootArgStrideInBytes = ((RootArgStrideInBytes + 15) & ~15);
+	MaxRootArgOffset = RootArgOffsets.Last();
 
-	WorkGraph->Finalize();
+	WorkGraphSubobject->Finalize();
 
+	ID3D12Device9* Device9 = Device->GetDevice9();
 	HRESULT HResult = Device9->CreateStateObject(StateObjectDesc, IID_PPV_ARGS(StateObject.GetInitReference()));
 	checkf(SUCCEEDED(HResult), TEXT("Failed to create work graph state object. Result=%08x"), HResult);
 
@@ -113,22 +132,6 @@ FD3D12WorkGraphPipelineState::FD3D12WorkGraphPipelineState(FD3D12Device* Device,
 	D3D12_WORK_GRAPH_MEMORY_REQUIREMENTS MemoryRequirements = {};
 	WorkGraphProperties->GetWorkGraphMemoryRequirements(WorkGraphIndex, &MemoryRequirements);
 
-	MaxRootArgOffset = 0;
-	RootArgOffsets.Reset();
-	if (bIsShaderBundle)
-	{
-		RootArgOffsets.AddDefaulted(ShaderBundleNum);
-		for (int32 Index = 0; Index < ShaderBundleNum; ++Index)
-		{
-			if (LocalNodeShaders[Index])
-			{
-				uint32 NodeIndex = WorkGraphProperties->GetNodeIndex(WorkGraphIndex, D3D12_NODE_ID{ *Initializer.GetShaderBundleNodeName(), (uint32)Index });
-				RootArgOffsets[Index] = WorkGraphProperties->GetNodeLocalRootArgumentsTableIndex(WorkGraphIndex, NodeIndex);
-				MaxRootArgOffset = FMath::Max(MaxRootArgOffset, RootArgOffsets[Index]);
-			}
-		}
-	}
-
 	ID3D12Resource* BackingMemoryBufferResource = nullptr;
 	{
 		CD3DX12_RESOURCE_DESC BufferDesc = CD3DX12_RESOURCE_DESC::Buffer(MemoryRequirements.MaxSizeInBytes, D3D12_RESOURCE_FLAG_ALLOW_UNORDERED_ACCESS, 65536ull);
@@ -339,25 +342,26 @@ static void RecordBindings(
 	FD3D12ExplicitDescriptorCache& TransientDescriptorCache,
 	FShaderBundleBinderOps& BinderOps,
 	uint32 WorkerIndex,
-	FRHIWorkGraphShader* WorkGraphShaderRHI,
+	FRHIShader* ShaderRHI,
+	FD3D12ShaderData const* D3D12ShaderData,
 	FRHIBatchedShaderParameters const& Parameters,
+	FD3D12RootSignature const* LocalRootSignature,
 	FAllocatedConstantBuffer const& SharedConstantBuffer,
 	FUint32Vector4 const& Constants,
 	TArrayView<uint32> RootArgs
 )
 {
-	FD3D12WorkGraphShader* D3D12WorkGraphShader = static_cast<FD3D12WorkGraphShader*>(WorkGraphShaderRHI);
-	const uint32 NumSMPs = D3D12WorkGraphShader->ResourceCounts.NumSamplers;
-	const uint32 NumSRVs = D3D12WorkGraphShader->ResourceCounts.NumSRVs;
-	const uint32 NumCBVs = D3D12WorkGraphShader->ResourceCounts.NumCBs;
-	const uint32 NumUAVs = D3D12WorkGraphShader->ResourceCounts.NumUAVs;
+	const uint32 NumSMPs = D3D12ShaderData->ResourceCounts.NumSamplers;
+	const uint32 NumSRVs = D3D12ShaderData->ResourceCounts.NumSRVs;
+	const uint32 NumCBVs = D3D12ShaderData->ResourceCounts.NumCBs;
+	const uint32 NumUAVs = D3D12ShaderData->ResourceCounts.NumUAVs;
 
 	// With shader root constants, we should never hit this expensive path!
 	// If we hit this, check if the shaders in the bundle had loose
 	// uniform parameters added to it recently, falling into this path.
-	check(!D3D12WorkGraphShader->UsesGlobalUniformBuffer() || SharedConstantBuffer.ConstantBuffer != nullptr);
+	check(!D3D12ShaderData->UsesGlobalUniformBuffer() || SharedConstantBuffer.ConstantBuffer != nullptr);
 
-	FWorkGraphShaderBundleBinder BundleBinder(Context, Frequency, BinderOps, D3D12WorkGraphShader);
+	FWorkGraphShaderBundleBinder BundleBinder(Context, Frequency, BinderOps, D3D12ShaderData);
 
 	FD3D12UniformBuffer* BundleUniformBuffers[MAX_CBS] = { nullptr };
 
@@ -394,7 +398,7 @@ static void RecordBindings(
 		}
 	}
 
-	UE::RHICore::ApplyStaticUniformBuffers(WorkGraphShaderRHI, Context.GetStaticUniformBuffers(),
+	UE::RHICore::ApplyStaticUniformBuffers(ShaderRHI, Context.GetStaticUniformBuffers(),
 		[&](int32 BufferIndex, FRHIUniformBuffer* Buffer)
 		{
 			BundleUniformBuffers[BufferIndex] = Context.RetrieveObject<FD3D12UniformBuffer>(Buffer);
@@ -402,7 +406,7 @@ static void RecordBindings(
 
 
 	uint32 FakeDirtyUniformBuffers = ~(0u);
-	UE::RHI::Private::SetUniformBufferResourcesFromTables(BundleBinder, *WorkGraphShaderRHI, FakeDirtyUniformBuffers, BundleUniformBuffers
+	UE::RHI::Private::SetUniformBufferResourcesFromTables(BundleBinder, *ShaderRHI, FakeDirtyUniformBuffers, BundleUniformBuffers
 #if ENABLE_RHI_VALIDATION
 		, Context.Tracker
 #endif
@@ -433,18 +437,18 @@ static void RecordBindings(
 		return (ExpectedMask & BoundMask) == ExpectedMask;
 	};
 
-	check(IsCompleteBinding(D3D12WorkGraphShader->ResourceCounts.NumSRVs, BundleBinder.BoundSRVMask));
-	check(IsCompleteBinding(D3D12WorkGraphShader->ResourceCounts.NumUAVs, BundleBinder.BoundUAVMask));
-	check(IsCompleteBinding(D3D12WorkGraphShader->ResourceCounts.NumCBs, BundleBinder.BoundCBVMask));
-	check(IsCompleteBinding(D3D12WorkGraphShader->ResourceCounts.NumSamplers, BundleBinder.BoundSamplerMask));
+	check(IsCompleteBinding(D3D12ShaderData->ResourceCounts.NumSRVs, BundleBinder.BoundSRVMask));
+	check(IsCompleteBinding(D3D12ShaderData->ResourceCounts.NumUAVs, BundleBinder.BoundUAVMask));
+	check(IsCompleteBinding(D3D12ShaderData->ResourceCounts.NumCBs, BundleBinder.BoundCBVMask));
+	check(IsCompleteBinding(D3D12ShaderData->ResourceCounts.NumSamplers, BundleBinder.BoundSamplerMask));
 
 	if (NumSRVs > 0)
 	{
 		const int32 DescriptorTableBaseIndex = TransientDescriptorCache.Allocate(BundleBinder.LocalSRVs, NumSRVs, D3D12_DESCRIPTOR_HEAP_TYPE_CBV_SRV_UAV, WorkerIndex);
 		check(DescriptorTableBaseIndex != INDEX_NONE);
 
-		const uint32 BindSlot = D3D12WorkGraphShader->RootSignature->SRVRDTBindSlot(SF_Compute);
-		const uint32 BindSlotOffset = D3D12WorkGraphShader->RootSignature->GetBindSlotOffsetInBytes(BindSlot) / 4;
+		const uint32 BindSlot = LocalRootSignature->SRVRDTBindSlot(Frequency);
+		const uint32 BindSlotOffset = LocalRootSignature->GetBindSlotOffsetInBytes(BindSlot) / 4;
 
 		const D3D12_GPU_DESCRIPTOR_HANDLE ResourceDescriptorTableBaseGPU = TransientDescriptorCache.ViewHeap.GetDescriptorGPU(DescriptorTableBaseIndex);
 		FMemory::Memcpy(&RootArgs[BindSlotOffset], &ResourceDescriptorTableBaseGPU, sizeof(ResourceDescriptorTableBaseGPU));
@@ -455,8 +459,8 @@ static void RecordBindings(
 		const int32 DescriptorTableBaseIndex = TransientDescriptorCache.AllocateDeduplicated(BundleBinder.SamplerVersions, BundleBinder.LocalSamplers, NumSMPs, D3D12_DESCRIPTOR_HEAP_TYPE_SAMPLER, WorkerIndex);
 		check(DescriptorTableBaseIndex != INDEX_NONE);
 
-		const uint32 BindSlot = D3D12WorkGraphShader->RootSignature->SamplerRDTBindSlot(SF_Compute);
-		const uint32 BindSlotOffset = D3D12WorkGraphShader->RootSignature->GetBindSlotOffsetInBytes(BindSlot) / 4;
+		const uint32 BindSlot = LocalRootSignature->SamplerRDTBindSlot(Frequency);
+		const uint32 BindSlotOffset = LocalRootSignature->GetBindSlotOffsetInBytes(BindSlot) / 4;
 
 		const D3D12_GPU_DESCRIPTOR_HANDLE ResourceDescriptorTableBaseGPU = TransientDescriptorCache.SamplerHeap.GetDescriptorGPU(DescriptorTableBaseIndex);
 		FMemory::Memcpy(&RootArgs[BindSlotOffset], &ResourceDescriptorTableBaseGPU, sizeof(ResourceDescriptorTableBaseGPU));
@@ -468,8 +472,8 @@ static void RecordBindings(
 		//const int32 DescriptorTableBaseIndex = TransientDescriptorCache.Allocate(BundleBinder.LocalUAVs, NumUAVs, D3D12_DESCRIPTOR_HEAP_TYPE_CBV_SRV_UAV, WorkerIndex);
 		check(DescriptorTableBaseIndex != INDEX_NONE);
 
-		const uint32 BindSlot = D3D12WorkGraphShader->RootSignature->UAVRDTBindSlot(SF_Compute);
-		const uint32 BindSlotOffset = D3D12WorkGraphShader->RootSignature->GetBindSlotOffsetInBytes(BindSlot) / 4;
+		const uint32 BindSlot = LocalRootSignature->UAVRDTBindSlot(Frequency);
+		const uint32 BindSlotOffset = LocalRootSignature->GetBindSlotOffsetInBytes(BindSlot) / 4;
 
 		const D3D12_GPU_DESCRIPTOR_HANDLE ResourceDescriptorTableBaseGPU = TransientDescriptorCache.ViewHeap.GetDescriptorGPU(DescriptorTableBaseIndex);
 		FMemory::Memcpy(&RootArgs[BindSlotOffset], &ResourceDescriptorTableBaseGPU, sizeof(ResourceDescriptorTableBaseGPU));
@@ -477,8 +481,8 @@ static void RecordBindings(
 
 	if (SharedConstantBuffer.ConstantBuffer != nullptr)
 	{
-		const uint32 BindSlot = D3D12WorkGraphShader->RootSignature->CBVRDBindSlot(SF_Compute, 0);
-		const uint32 BindSlotOffset = D3D12WorkGraphShader->RootSignature->GetBindSlotOffsetInBytes(BindSlot) / 4;
+		const uint32 BindSlot = LocalRootSignature->CBVRDBindSlot(Frequency, 0);
+		const uint32 BindSlotOffset = LocalRootSignature->GetBindSlotOffsetInBytes(BindSlot) / 4;
 
 		D3D12_GPU_VIRTUAL_ADDRESS Address = SharedConstantBuffer.ResourceLocation.GetGPUVirtualAddress();
 		FMemory::Memcpy(&RootArgs[BindSlotOffset], &Address, sizeof(Address));
@@ -486,8 +490,8 @@ static void RecordBindings(
 
 	for (uint32 CBVIndex = 0; CBVIndex < NumCBVs; ++CBVIndex)
 	{
-		const uint32 BindSlot = D3D12WorkGraphShader->RootSignature->CBVRDBindSlot(SF_Compute, CBVIndex);
-		const uint32 BindSlotOffset = D3D12WorkGraphShader->RootSignature->GetBindSlotOffsetInBytes(BindSlot) / 4;
+		const uint32 BindSlot = LocalRootSignature->CBVRDBindSlot(Frequency, CBVIndex);
+		const uint32 BindSlotOffset = LocalRootSignature->GetBindSlotOffsetInBytes(BindSlot) / 4;
 
 		FD3D12UniformBuffer* UniformBuffer = BundleUniformBuffers[CBVIndex];
 		if (UniformBuffer)
@@ -497,10 +501,10 @@ static void RecordBindings(
 		}
 	}
 
-	const int8 BindSlot = D3D12WorkGraphShader->RootSignature->GetRootConstantsSlot();
+	const int8 BindSlot = LocalRootSignature->GetRootConstantsSlot();
 	if (BindSlot != -1)
 	{
-		const uint32 BindSlotOffset = D3D12WorkGraphShader->RootSignature->GetBindSlotOffsetInBytes(BindSlot) / 4;
+		const uint32 BindSlotOffset = LocalRootSignature->GetBindSlotOffsetInBytes(BindSlot) / 4;
 
 		RootArgs[BindSlotOffset] = Constants.X;
 		RootArgs[BindSlotOffset + 1] = Constants.Y;
@@ -660,16 +664,18 @@ void FD3D12CommandContext::DispatchWorkGraphShaderBundle(
 	TArray<uint32> ValidRecords;
 	ValidRecords.Reserve(NumRecords);
 	TArray<FRHIWorkGraphShader*> LocalNodeShaders;
-	LocalNodeShaders.Reserve(NumRecords);
+	LocalNodeShaders.Reserve(NumRecords + 1);
+	LocalNodeShaders.Add(WorkGraphGlobalShaderRHI);
 
 	for (int32 DispatchIndex = 0; DispatchIndex < NumRecords; ++DispatchIndex)
 	{
 		const FRHIShaderBundleComputeDispatch& Dispatch = Dispatches[DispatchIndex];
 		FRHIWorkGraphShader* Shader = Dispatch.IsValid() ? Dispatch.WorkGraphShader : nullptr;
 		LocalNodeShaders.Add(Shader);
+
 		if (Shader != nullptr)
 		{
-			ValidRecords.Emplace(uint32(DispatchIndex));
+			ValidRecords.Add(uint32(DispatchIndex));
 
 			if (FD3D12WorkGraphShader* D3D12Shader = FD3D12DynamicRHI::ResourceCast(Shader))
 			{
@@ -681,9 +687,13 @@ void FD3D12CommandContext::DispatchWorkGraphShaderBundle(
 	const int32 NumValidRecords = ValidRecords.Num();
 
 	FWorkGraphPipelineStateInitializer Initializer;
-	Initializer.SetProgramName(TEXT("ShaderBundleWorkGraph"));
-	Initializer.SetShader(WorkGraphGlobalShaderRHI);
-	Initializer.SetShaderBundleNodeTable(LocalNodeShaders, TEXT("ShaderBundleNode"));
+	Initializer.SetProgramName(TEXT("ShaderBundle"));
+	TArray<FWorkGraphPipelineStateInitializer::FNameMap> NameTable;
+	NameTable.Add(FWorkGraphPipelineStateInitializer::FNameMap(TEXT("WorkGraphMainCS"), TEXT("WorkGraphMainCS")));
+	NameTable.Add(FWorkGraphPipelineStateInitializer::FNameMap(TEXT("MainCS"), TEXT("ShaderBundleNode")));
+	NameTable.Add(FWorkGraphPipelineStateInitializer::FNameMap(TEXT("MicropolyRasterize"), TEXT("ShaderBundleNode")));
+	Initializer.SetNameTable(NameTable);
+	Initializer.SetShaderTable(LocalNodeShaders);
 
 	FWorkGraphPipelineState* WorkGraphPipelineState = PipelineStateCache::GetAndOrCreateWorkGraphPipelineState(RHICmdList, Initializer);
 	FD3D12WorkGraphPipelineState* Pipeline = static_cast<FD3D12WorkGraphPipelineState*>(GetRHIWorkGraphPipelineState(WorkGraphPipelineState));
@@ -721,13 +731,18 @@ void FD3D12CommandContext::DispatchWorkGraphShaderBundle(
 
 	auto RecordTask = [this, &LocalRootArgs, Pipeline, &TransientDescriptorCache, &ValidRecords, &Dispatches, &BinderOps, &SharedConstantBuffer](FTaskContext& Context, int32 RecordIndex)
 	{
-		uint32 DispatchIndex = ValidRecords[RecordIndex];
+		const uint32 DispatchIndex = ValidRecords[RecordIndex];
 		const FRHIShaderBundleComputeDispatch& Dispatch = Dispatches[DispatchIndex];
+		check(Dispatch.IsValid());
 
-		check(Pipeline->RootArgOffsets.IsValidIndex(DispatchIndex));
-		uint32 RootArgOffset = Pipeline->RootArgOffsets[DispatchIndex];
+		const uint32 ShaderTableIndex = RecordIndex + 1;
+		check(Pipeline->RootArgOffsets.IsValidIndex(ShaderTableIndex));
+		uint32 RootArgOffset = Pipeline->RootArgOffsets[ShaderTableIndex];
 		check((Pipeline->RootArgStrideInBytes / 4) * (RootArgOffset + 1) <= (uint32)LocalRootArgs.Num());
 
+		FD3D12WorkGraphShader* D3D12WorkGraphShader = FD3D12DynamicRHI::ResourceCast(Dispatch.WorkGraphShader);
+		FD3D12RootSignature const* LocalRootSignature = D3D12WorkGraphShader->RootSignature;
+
 		RecordBindings(
 			*this,
 			SF_Compute,
@@ -735,7 +750,9 @@ void FD3D12CommandContext::DispatchWorkGraphShaderBundle(
 			BinderOps[Context.WorkerIndex],
 			Context.WorkerIndex,
 			Dispatch.WorkGraphShader,
+			D3D12WorkGraphShader,
 			*Dispatch.Parameters,
+			LocalRootSignature,
 			SharedConstantBuffer,
 			Dispatch.Constants,
 			MakeArrayView(&LocalRootArgs[RootArgOffset * Pipeline->RootArgStrideInBytes / 4], Pipeline->RootArgStrideInBytes / 4)
@@ -749,6 +766,44 @@ void FD3D12CommandContext::DispatchWorkGraphShaderBundle(
 
 	ParallelForWithExistingTaskContext(TEXT("DispatchShaderBundle"), MakeArrayView(TaskContexts), ValidRecords.Num(), ItemsPerTask, RecordTask);
 
+	// Bind RecordArgBuffer
+	FD3D12Buffer* RecordArgBufferPtr = FD3D12DynamicRHI::ResourceCast(RecordArgBuffer);
+
+	D3D12_SHADER_RESOURCE_VIEW_DESC SRVDesc = {};
+	SRVDesc.ViewDimension = D3D12_SRV_DIMENSION_BUFFER;
+	SRVDesc.Format = DXGI_FORMAT_R32_TYPELESS;
+	SRVDesc.Shader4ComponentMapping = D3D12_DEFAULT_SHADER_4_COMPONENT_MAPPING;
+	SRVDesc.Buffer.FirstElement = 0;
+	SRVDesc.Buffer.NumElements = RecordArgBufferPtr->GetSize() >> 2u;
+	SRVDesc.Buffer.StructureByteStride = 0;
+	SRVDesc.Buffer.Flags = D3D12_BUFFER_SRV_FLAG_RAW;
+
+	TSharedPtr<FD3D12ShaderResourceView> RecordArgBufferSRV;
+	RecordArgBufferSRV = MakeShared<FD3D12ShaderResourceView>(GetParentDevice(), nullptr); // Always single GPU object, so FirstLinkedObject is nullptr
+	RecordArgBufferSRV->CreateView(RecordArgBufferPtr, SRVDesc, FD3D12ShaderResourceView::EFlags::None);
+
+	uint32 RecordArgBufferBindlessHandle = 0;
+
+	if (bBindlessResources)
+	{
+		RecordArgBufferBindlessHandle = RecordArgBufferSRV->GetBindlessHandle().GetIndex();
+		check(RecordArgBufferBindlessHandle != INDEX_NONE);
+	}
+	else
+	{
+		D3D12_CPU_DESCRIPTOR_HANDLE LocalSRVs[MAX_SRVS];
+		LocalSRVs[WorkGraphGlobalShader->RecordArgBufferParam.GetBaseIndex()] = RecordArgBufferSRV->GetOfflineCpuHandle();
+
+		const int32 DescriptorTableBaseIndex = TransientDescriptorCache.Allocate(LocalSRVs, 1, D3D12_DESCRIPTOR_HEAP_TYPE_CBV_SRV_UAV, 0);
+		const D3D12_GPU_DESCRIPTOR_HANDLE ResourceDescriptorTableBaseGPU = TransientDescriptorCache.ViewHeap.GetDescriptorGPU(DescriptorTableBaseIndex);
+
+		const uint32 BindSlot = WorkGraphGlobalShaderRHI->RootSignature->SRVRDTBindSlot(SF_Compute);
+		const uint32 BindSlotOffset = WorkGraphGlobalShaderRHI->RootSignature->GetBindSlotOffsetInBytes(BindSlot) / 4;
+
+		TArrayView<uint32> RootArgSlice = MakeArrayView(&LocalRootArgs[Pipeline->RootArgOffsets[0] * Pipeline->RootArgStrideInBytes / 4], Pipeline->RootArgStrideInBytes / 4);
+		FMemory::Memcpy(&RootArgSlice[BindSlotOffset], &ResourceDescriptorTableBaseGPU, sizeof(ResourceDescriptorTableBaseGPU));
+	}
+
 	// Upload local root arguments table.
 	D3D12_GPU_VIRTUAL_ADDRESS_RANGE_AND_STRIDE NodeLocalRootArgumentsTable{ 0, 0, 0 };
 	if (ValidRecords.Num() && LocalRootArgs.Num())
@@ -812,52 +867,14 @@ void FD3D12CommandContext::DispatchWorkGraphShaderBundle(
 	// Note this only uses transient heap for non-bindless.
 	SetExplicitDescriptorCache(TransientDescriptorCache);
 
-	// Gather root arguments for global work graph.
-	FD3D12Buffer* RecordArgBufferPtr = FD3D12DynamicRHI::ResourceCast(RecordArgBuffer);
-
-	D3D12_SHADER_RESOURCE_VIEW_DESC SRVDesc = {};
-	SRVDesc.ViewDimension = D3D12_SRV_DIMENSION_BUFFER;
-	SRVDesc.Format = DXGI_FORMAT_R32_TYPELESS;
-	SRVDesc.Shader4ComponentMapping = D3D12_DEFAULT_SHADER_4_COMPONENT_MAPPING;
-	SRVDesc.Buffer.FirstElement = 0;
-	SRVDesc.Buffer.NumElements = RecordArgBufferPtr->GetSize() >> 2u;
-	SRVDesc.Buffer.StructureByteStride = 0;
-	SRVDesc.Buffer.Flags = D3D12_BUFFER_SRV_FLAG_RAW;
-
-	TSharedPtr<FD3D12ShaderResourceView> RecordArgBufferSRV;
-	RecordArgBufferSRV = MakeShared<FD3D12ShaderResourceView>(GetParentDevice(), nullptr); // Always single GPU object, so FirstLinkedObject is nullptr
-	RecordArgBufferSRV->CreateView(RecordArgBufferPtr, SRVDesc, FD3D12ShaderResourceView::EFlags::None);
-
-	uint32 RecordArgBufferBindlessHandle = 0;
-	
 	if (bBindlessResources)
 	{
-		RecordArgBufferBindlessHandle = RecordArgBufferSRV->GetBindlessHandle().GetIndex();
-		check(RecordArgBufferBindlessHandle != INDEX_NONE);
-
 		StateCache.ApplyBindlessResources(nullptr, SF_Compute, SF_NumStandardFrequencies);
-		
-		GraphicsCommandList()->SetComputeRootSignature(WorkGraphGlobalShaderRHI->RootSignature->GetRootSignature());
 	}
-	else
-	{
-		int32 DispatchSRVTable = INDEX_NONE;
-		{
-			D3D12_CPU_DESCRIPTOR_HANDLE LocalSRVs[MAX_SRVS];
-			LocalSRVs[WorkGraphGlobalShader->RecordArgBufferParam.GetBaseIndex()] = RecordArgBufferSRV->GetOfflineCpuHandle();
-			DispatchSRVTable = TransientDescriptorCache.Allocate(LocalSRVs, 1, D3D12_DESCRIPTOR_HEAP_TYPE_CBV_SRV_UAV, 0);
-		}
-		check(DispatchSRVTable != INDEX_NONE);
 
-		const D3D12_GPU_DESCRIPTOR_HANDLE DispatchSRVHandle = TransientDescriptorCache.ViewHeap.GetDescriptorGPU(DispatchSRVTable);
-		const uint32 SRVBindSlot = WorkGraphGlobalShaderRHI->RootSignature->SRVRDTBindSlot(SF_Compute);
-		check(SRVBindSlot != 0xFF);
+	GraphicsCommandList()->SetComputeRootSignature(Pipeline->RootSignature);
 
-		GraphicsCommandList()->SetComputeRootSignature(WorkGraphGlobalShaderRHI->RootSignature->GetRootSignature());
-		GraphicsCommandList()->SetComputeRootDescriptorTable(SRVBindSlot, DispatchSRVHandle);
-	}
-	
-	// Kick off the work graph	
+	// Kick off the work graph
 	D3D12_SET_PROGRAM_DESC SetProgramDesc = {};
 	SetProgramDesc.Type = D3D12_PROGRAM_TYPE_WORK_GRAPH;
 	SetProgramDesc.WorkGraph.ProgramIdentifier = Pipeline->ProgramIdentifier;
diff --git a/Engine/Source/Runtime/D3D12RHI/Private/D3D12WorkGraph.h b/Engine/Source/Runtime/D3D12RHI/Private/D3D12WorkGraph.h
--- a/Engine/Source/Runtime/D3D12RHI/Private/D3D12WorkGraph.h
+++ b/Engine/Source/Runtime/D3D12RHI/Private/D3D12WorkGraph.h
@@ -5,9 +5,7 @@
 #include "D3D12RHICommon.h"
 #include "RHIResources.h"
 
-class FD3D12Buffer;
 class FD3D12Device;
-class FD3D12WorkGraphShader;
 class FWorkGraphPipelineStateInitializer;
 
 struct FD3D12WorkGraphPipelineState : public FRHIWorkGraphPipelineState
@@ -17,19 +15,19 @@ public:
 
 	FD3D12WorkGraphPipelineState(FD3D12Device* Device, const FWorkGraphPipelineStateInitializer& Initializer);
 
-	TRefCountPtr<FD3D12WorkGraphShader>	Shader;
-	TArray<TRefCountPtr<FD3D12WorkGraphShader>> LocalNodeShaders;
-
 #if D3D12_RHI_WORKGRAPHS
+	TRefCountPtr<ID3D12RootSignature> RootSignature;
 	TRefCountPtr<ID3D12StateObject>	StateObject;
-	D3D12_PROGRAM_IDENTIFIER ProgramIdentifier = {};
 	
+	D3D12_PROGRAM_IDENTIFIER ProgramIdentifier = {};
 	D3D12_GPU_VIRTUAL_ADDRESS_RANGE BackingMemoryAddressRange = {};
 
 	uint32 RootArgStrideInBytes = 0;
 	uint32 MaxRootArgOffset = 0;
 	TArray<uint32> RootArgOffsets;
 
+	TMap<FString, uint32> NodeCountPerName;
+
 	bool bInitialized = false;
 #endif
 };
diff --git a/Engine/Source/Runtime/RHI/Private/PipelineStateCache.cpp b/Engine/Source/Runtime/RHI/Private/PipelineStateCache.cpp
--- a/Engine/Source/Runtime/RHI/Private/PipelineStateCache.cpp
+++ b/Engine/Source/Runtime/RHI/Private/PipelineStateCache.cpp
@@ -3329,7 +3329,7 @@ FComputePipelineState* PipelineStateCache::GetAndOrCreateComputePipelineState(FR
 
 inline void ValidateWorkGraphPipelineStateInitializer(const FWorkGraphPipelineStateInitializer& Initializer)
 {
-	check(Initializer.GetShader() != nullptr);
+	check(Initializer.GetShaderTable().Num() > 0);
 }
 
 FWorkGraphPipelineState* PipelineStateCache::GetAndOrCreateWorkGraphPipelineState(FRHIComputeCommandList& RHICmdList, const FWorkGraphPipelineStateInitializer& Initializer)
@@ -3342,7 +3342,7 @@ FWorkGraphPipelineState* PipelineStateCache::GetAndOrCreateWorkGraphPipelineStat
 
 	if (!bWasFound)
 	{
-		OutCachedState = new FWorkGraphPipelineState(Initializer.GetShader());
+		OutCachedState = new FWorkGraphPipelineState(Initializer.GetShaderTable()[0]);
 		OutCachedState->RHIPipeline = RHICreateWorkGraphPipelineState(Initializer);
 		GWorkGraphPipelineCache.Add(Initializer, OutCachedState);
 	}
diff --git a/Engine/Source/Runtime/RHI/Public/RHIResources.h b/Engine/Source/Runtime/RHI/Public/RHIResources.h
--- a/Engine/Source/Runtime/RHI/Public/RHIResources.h
+++ b/Engine/Source/Runtime/RHI/Public/RHIResources.h
@@ -4502,66 +4502,87 @@ inline uint64 ComputeShaderTableHash(const TArrayView<TShaderType*>& ShaderTable
 class FWorkGraphPipelineStateSignature
 {
 public:
+	FWorkGraphPipelineStateSignature() = default;
+
 	bool operator==(const FWorkGraphPipelineStateSignature& Rhs) const
 	{
-		return BaseHash == Rhs.BaseHash && NameHash == Rhs.NameHash && ShaderBundleNodeHash == Rhs.ShaderBundleNodeHash;
+		return NameHash == Rhs.NameHash && 
+			NameTableHash == Rhs.NameTableHash && 
+			ShaderTableHash == Rhs.ShaderTableHash;
 	}
 
 	friend uint32 GetTypeHash(const FWorkGraphPipelineStateSignature& Initializer)
 	{
-		return GetTypeHash(Initializer.BaseHash) ^ GetTypeHash(Initializer.NameHash) ^ GetTypeHash(Initializer.ShaderBundleNodeHash);
+		return GetTypeHash(Initializer.NameHash) ^ 
+			GetTypeHash(Initializer.NameTableHash) ^ 
+			GetTypeHash(Initializer.ShaderTableHash);
 	}
 
+	uint64 GetNameHash() const { return NameHash; }
+	uint64 GetNameTableHash() const { return NameTableHash; }
+	uint64 GetShaderTableHash() const { return ShaderTableHash; }
+
 protected:
-	uint64 BaseHash = 0;
 	uint64 NameHash = 0;
-	uint64 ShaderBundleNodeHash = 0;
+	uint64 NameTableHash = 0;
+	uint64 ShaderTableHash = 0;
 };
 
 class FWorkGraphPipelineStateInitializer : public FWorkGraphPipelineStateSignature
 {
 public:
 	FWorkGraphPipelineStateInitializer() = default;
-	FWorkGraphPipelineStateRHIRef BasePipeline;
 
 	void SetProgramName(TCHAR const* InProgramName)
 	{
 		ProgramName = InProgramName;
-		NameHash = ComputeNameHash();
+		NameHash = GetTypeHash(ProgramName);
 	}
 
-	FString const& GetProgramName() const { return ProgramName; }
-
-	void SetShader(FRHIWorkGraphShader* InShader)
+	/** A single mapping of shader export name to work graph node name. An array of these defines the full mapping. */
+	struct FNameMap
 	{
-		BaseShader = InShader;
-		FMemory::Memcpy(&BaseHash, InShader->GetHash().Hash, sizeof(BaseHash));	// 64 bits from the shader SHA1
-	}
+		FString ExportName;
+		FString NodeName;
+		uint32 ExportNameHash;
 
-	FRHIWorkGraphShader* GetShader() const { return BaseShader; }
+		FNameMap(FString const& InExportName, FString const& InNodeName)
+			: ExportName(InExportName)
+			, NodeName(InNodeName)
+			, ExportNameHash(GetTypeHash(InExportName))
+		{}
 
-	void SetShaderBundleNodeTable(const TArrayView<FRHIWorkGraphShader*>& InShaders, TCHAR const* InNodeName, uint64 Hash = 0)
-	{
-		ShaderBundleNodeTable = InShaders;
-		ShaderBundleNodeHash = Hash ? Hash : ComputeShaderTableHash(InShaders);
+		friend uint32 GetTypeHash(FNameMap const& NameMap)
+		{
+			return HashCombineFast(GetTypeHash(NameMap.ExportName), GetTypeHash(NameMap.NodeName));
+		}
+	};
 
-		ShaderBundleNodeName = InNodeName;
-		NameHash = ComputeNameHash();
+	void SetNameTable(const TArrayView<FNameMap> InNameMaps, uint64 Hash = 0)
+	{
+		NameMaps = InNameMaps;
+		NameTableHash = Hash ? Hash : GetArrayHash(NameMaps.GetData(), NameMaps.Num());
 	}
 
-	const TArrayView<FRHIWorkGraphShader*>& GetShaderBundleNodeTable() const { return ShaderBundleNodeTable; }
-	const FString& GetShaderBundleNodeName() const { return ShaderBundleNodeName; }
-
-private:
-	uint64 ComputeNameHash() const
+	void SetShaderTable(const TArrayView<FRHIWorkGraphShader*>& InShaders, int32 InRootShaderIndex = -1, uint64 Hash = 0)
 	{
-		return HashCombineFast(GetTypeHash(ProgramName), GetTypeHash(ShaderBundleNodeName));
+		ShaderTable = InShaders;
+		ShaderTableHash = Hash ? Hash : ComputeShaderTableHash(InShaders);
+		// RootShaderIndex doesn't need adding to a hash because if used correctly it is an artifact of the shader table.
+		// (Only one shader in the array can have a global root signature).
+		RootShaderIndex = InRootShaderIndex;
 	}
 
-	FRHIWorkGraphShader* BaseShader = nullptr;
+	FString const& GetProgramName() const { return ProgramName; }
+	TArray<FNameMap> const& GetNameTable() const { return NameMaps; }
+	int32 GetRootShaderIndex() const { return RootShaderIndex; }
+	TArrayView<FRHIWorkGraphShader*> const& GetShaderTable() const { return ShaderTable; }
+
+private:
 	FString ProgramName;
-	TArrayView<FRHIWorkGraphShader*> ShaderBundleNodeTable;
-	FString ShaderBundleNodeName;
+	TArray<FNameMap> NameMaps;
+	int32 RootShaderIndex = -1;
+	TArrayView<FRHIWorkGraphShader*> ShaderTable;
 };
 
 class FRayTracingPipelineStateSignature
diff --git a/Engine/Source/Runtime/RenderCore/Private/ShaderBundles.cpp b/Engine/Source/Runtime/RenderCore/Private/ShaderBundles.cpp
--- a/Engine/Source/Runtime/RenderCore/Private/ShaderBundles.cpp
+++ b/Engine/Source/Runtime/RenderCore/Private/ShaderBundles.cpp
@@ -36,7 +36,7 @@ void FDispatchShaderBundleCS::ModifyCompilationEnvironment(const FGlobalShaderPe
 }
 
 
-IMPLEMENT_GLOBAL_SHADER(FDispatchShaderBundleWorkGraph, "/Engine/Private/ShaderBundleWorkGraphDispatch.usf", "WorkGraphMainCS", SF_WorkGraphRoot);
+IMPLEMENT_GLOBAL_SHADER(FDispatchShaderBundleWorkGraph, "/Engine/Private/ShaderBundleWorkGraphDispatch.usf", "WorkGraphMainCS", SF_WorkGraphComputeNode);
 
 bool FDispatchShaderBundleWorkGraph::ShouldCompilePermutation(const FGlobalShaderPermutationParameters& Parameters)
 {
