Cache the value r.PSOPrecache.D3D12.DriverCacheAware in D3D12RHI as it could be changed by unrelated hotfix logic and cause a crash.

The current hotfix and config's dynamic layers logic make it so that hotfixing device profiles results in unrelated CVars being popped/pushed. If the value of this CVar changes at runtime, it will cause crashes in the precaching system as the computed PSO hashes are going to differ.

#rnx

diff --git a/Engine/Source/Runtime/D3D12RHI/Private/D3D12RHI.cpp b/Engine/Source/Runtime/D3D12RHI/Private/D3D12RHI.cpp
--- a/Engine/Source/Runtime/D3D12RHI/Private/D3D12RHI.cpp
+++ b/Engine/Source/Runtime/D3D12RHI/Private/D3D12RHI.cpp
@@ -42,6 +42,14 @@ TAutoConsoleVariable<int32> CVarD3D12ZeroBufferSizeInMB(
 	ECVF_ReadOnly
 	);
 
+static bool GPSOPrecacheD3D12DriverCacheAware = false;
+static FAutoConsoleVariableRef CVarPSOPrecacheD3D12DriverCacheAware(
+	TEXT("r.PSOPrecache.D3D12.DriverCacheAware"),
+	GPSOPrecacheD3D12DriverCacheAware,
+	TEXT("If enabled, the PSO precaching system will not precache PSOs that the D3D12 graphics driver considers similar for caching, i.e. it will not precache PSOs that while technically different will still result in a driver cache hit.\n")
+	TEXT("This is not implemented for all GPU vendors and can result in performance issues or cache misses if the heuristics the engine uses does not match the graphics driver's behavior that decides whether a PSO is in the cache or not."),
+	ECVF_ReadOnly);
+
 FD3D12DynamicRHI* FD3D12DynamicRHI::SingleD3DRHI = nullptr;
 
 FD3D12WorkaroundFlags GD3D12WorkaroundFlags;
@@ -78,6 +86,8 @@ FD3D12DynamicRHI::FD3D12DynamicRHI(const TArray<TSharedPtr<FD3D12Adapter>>& Chos
 	ZeroBuffer = nullptr;
 #endif // PLATFORM_WINDOWS
 
+	bDriverCacheAwarePSOPrecaching = GPSOPrecacheD3D12DriverCacheAware;
+
 	GRHIGlobals.SupportsMultiDrawIndirect = true;
 
 	GRHISupportsMultithreading = true;
diff --git a/Engine/Source/Runtime/D3D12RHI/Private/D3D12RHIPrivate.h b/Engine/Source/Runtime/D3D12RHI/Private/D3D12RHIPrivate.h
--- a/Engine/Source/Runtime/D3D12RHI/Private/D3D12RHIPrivate.h
+++ b/Engine/Source/Runtime/D3D12RHI/Private/D3D12RHIPrivate.h
@@ -545,6 +545,11 @@ protected:
 
 	INTCExtensionContext* IntelExtensionContext = nullptr;
 
+	// See r.PSOPrecache.D3D12.DriverCacheAware. The CVar is cached
+	// to prevent hotfixes or config dynamic layers from altering the
+	// value at runtime, which is not supported.
+	bool bDriverCacheAwarePSOPrecaching = false;
+
 	/** A buffer in system memory containing all zeroes of the specified size. */
 	void* ZeroBuffer;
 	uint32 ZeroBufferSize;
diff --git a/Engine/Source/Runtime/D3D12RHI/Private/D3D12State.cpp b/Engine/Source/Runtime/D3D12RHI/Private/D3D12State.cpp
--- a/Engine/Source/Runtime/D3D12RHI/Private/D3D12State.cpp
+++ b/Engine/Source/Runtime/D3D12RHI/Private/D3D12State.cpp
@@ -8,14 +8,6 @@
 #include "D3D12RHIPrivate.h"
 #include "RHIUtilities.h"
 
-static int32 GPSOPrecacheD3D12DriverCacheAware = 0;
-static FAutoConsoleVariableRef CVarPSOPrecacheD3D12DriverCacheAware(
-	TEXT("r.PSOPrecache.D3D12.DriverCacheAware"),
-	GPSOPrecacheD3D12DriverCacheAware,
-	TEXT("If enabled, the PSO precaching system will not precache PSOs that the D3D12 graphics driver considers similar for caching, i.e. it will not precache PSOs that while technically different will still result in a driver cache hit.\n")
-	TEXT("This is not implemented for all GPU vendors and can result in performance issues or cache misses if the heuristics the engine uses does not match the graphics driver's behavior that decides whether a PSO is in the cache or not."),
-	ECVF_ReadOnly);
-
 // MSFT: Need to make sure sampler state is thread safe
 // Cache of Sampler States; we store pointers to both as we don't want the TMap to be artificially
 // modifying ref counts if not needed; so we manage that ourselves
@@ -498,7 +490,7 @@ bool FD3D12BlendState::GetInitializer(class FBlendStateInitializerRHI& Init)
 
 uint64 FD3D12DynamicRHI::RHIComputeStatePrecachePSOHash(const FGraphicsPipelineStateInitializer& Initializer)
 {
-	if (GPSOPrecacheD3D12DriverCacheAware)
+	if (bDriverCacheAwarePSOPrecaching)
 	{
 		if (IsRHIDeviceNVIDIA() && GMaxRHIFeatureLevel == ERHIFeatureLevel::SM6)
 		{
@@ -580,7 +572,7 @@ uint64 FD3D12DynamicRHI::RHIComputePrecachePSOHash(const FGraphicsPipelineStateI
 	// When compute precache PSO hash we assume a valid state precache PSO hash is already provided
 	checkf(Initializer.StatePrecachePSOHash != 0, TEXT("Initializer should have a valid state precache PSO hash set when computing the full initializer PSO hash"));
 
-	if (GPSOPrecacheD3D12DriverCacheAware)
+	if (bDriverCacheAwarePSOPrecaching)
 	{
 		if (IsRHIDeviceNVIDIA() && GMaxRHIFeatureLevel == ERHIFeatureLevel::SM6)
 		{
