Quick fix for missing Device9 build fail.

diff --git a/Engine/Source/Runtime/D3D12RHI/Private/D3D12WorkGraph.cpp b/Engine/Source/Runtime/D3D12RHI/Private/D3D12WorkGraph.cpp
--- a/Engine/Source/Runtime/D3D12RHI/Private/D3D12WorkGraph.cpp
+++ b/Engine/Source/Runtime/D3D12RHI/Private/D3D12WorkGraph.cpp
@@ -114,19 +114,19 @@ FD3D12WorkGraphPipelineState::FD3D12WorkGraphPipelineState(FD3D12Device* Device,
 
 	WorkGraphSubobject->Finalize();
 
-	ID3D12Device9* Device9 = Device->GetDevice9();
+	ID3D12Device9* Device9 = (ID3D12Device9*)Device->GetDevice();
 	HRESULT HResult = Device9->CreateStateObject(StateObjectDesc, IID_PPV_ARGS(StateObject.GetInitReference()));
-	checkf(SUCCEEDED(HResult), TEXT("Failed to create work graph state object. Result=%08x"), HResult);
+	checkf(SUCCEEDED(HResult), TEXT("Failed to create work graph state object. Result=%08x"), uint32(HResult));
 
 	TRefCountPtr<ID3D12StateObjectProperties1> PipelineProperties;
 	HResult = StateObject->QueryInterface(IID_PPV_ARGS(PipelineProperties.GetInitReference()));
-	checkf(SUCCEEDED(HResult), TEXT("Failed to query pipeline properties from the work graph pipeline state object. Result=%08x"), HResult);
+	checkf(SUCCEEDED(HResult), TEXT("Failed to query pipeline properties from the work graph pipeline state object. Result=%08x"), uint32(HResult));
 
 	ProgramIdentifier = PipelineProperties->GetProgramIdentifier(ProgramName);
 
 	TRefCountPtr<ID3D12WorkGraphProperties> WorkGraphProperties;
 	HResult = StateObject->QueryInterface(IID_PPV_ARGS(WorkGraphProperties.GetInitReference()));
-	checkf(SUCCEEDED(HResult), TEXT("Failed to query work graph properties from the work graph pipeline state object. Result=%08x"), HResult);
+	checkf(SUCCEEDED(HResult), TEXT("Failed to query work graph properties from the work graph pipeline state object. Result=%08x"), uint32(HResult));
 
 	UINT WorkGraphIndex = WorkGraphProperties->GetWorkGraphIndex(ProgramName);
 	D3D12_WORK_GRAPH_MEMORY_REQUIREMENTS MemoryRequirements = {};
@@ -145,7 +145,7 @@ FD3D12WorkGraphPipelineState::FD3D12WorkGraphPipelineState(FD3D12Device* Device,
 			NULL,
 			__uuidof(ID3D12Resource),
 			(void**)&BackingMemoryBufferResource);
-		checkf(SUCCEEDED(HResult), TEXT("Failed to allocate backing memory for work graph. Result=%08x"), HResult);
+		checkf(SUCCEEDED(HResult), TEXT("Failed to allocate backing memory for work graph. Result=%08x"), uint32(HResult));
 	}
 	BackingMemoryAddressRange.StartAddress = BackingMemoryBufferResource->GetGPUVirtualAddress();
 	BackingMemoryAddressRange.SizeInBytes = MemoryRequirements.MaxSizeInBytes;
