Expose if RHI supports Work Graph graphics nodes.
This is based on the Tier 1.1 support in the Preview Agility SDK.
We use D3D12_RHI_WORKGRAPHS_GRAPHICS to determine if the Preview Agility SDK is available but will not be needed once support is in a full SDK.
That define and the DDSPI bSupportsWorkGraphsTier1_1 value are for setting explicitly in change lists that test the Preview SDK.

diff --git a/Engine/Source/Runtime/D3D12RHI/Private/D3D12Adapter.cpp b/Engine/Source/Runtime/D3D12RHI/Private/D3D12Adapter.cpp
--- a/Engine/Source/Runtime/D3D12RHI/Private/D3D12Adapter.cpp
+++ b/Engine/Source/Runtime/D3D12RHI/Private/D3D12Adapter.cpp
@@ -1109,9 +1109,16 @@ void FD3D12Adapter::InitializeDevices()
 				if (D3D12Caps21.WorkGraphsTier != D3D12_WORK_GRAPHS_TIER_NOT_SUPPORTED)
 				{
 					UE_LOG(LogD3D12RHI, Log, TEXT("Work Graphs are supported"));
-					GRHISupportsShaderWorkGraphsTier1 = true;
+					GRHIGlobals.SupportsShaderWorkGraphsTier1 = true;
 				}
 #endif // D3D12_RHI_WORKGRAPHS
+#if D3D12_RHI_WORKGRAPHS_GRAPHICS
+				if (D3D12Caps21.WorkGraphsTier >= D3D12_WORK_GRAPHS_TIER_1_1)
+				{
+					UE_LOG(LogD3D12RHI, Log, TEXT("Work Graphs Tier 1.1 is supported"));
+					GRHIGlobals.SupportsShaderWorkGraphsTier1_1 = true;
+				}
+#endif // D3D12_RHI_WORKGRAPHS_GRAPHICS
 #endif // D3D12_MAX_FEATURE_OPTIONS
 			}
 #endif // PLATFORM_WINDOWS
diff --git a/Engine/Source/Runtime/D3D12RHI/Private/D3D12RHIDefinitions.h b/Engine/Source/Runtime/D3D12RHI/Private/D3D12RHIDefinitions.h
--- a/Engine/Source/Runtime/D3D12RHI/Private/D3D12RHIDefinitions.h
+++ b/Engine/Source/Runtime/D3D12RHI/Private/D3D12RHIDefinitions.h
@@ -20,8 +20,10 @@
 
 #if D3D12_MAX_COMMANDLIST_INTERFACE >= 10
 	#define D3D12_RHI_WORKGRAPHS 1
+	#define D3D12_RHI_WORKGRAPHS_GRAPHICS 0
 #else
 	#define D3D12_RHI_WORKGRAPHS 0
+	#define D3D12_RHI_WORKGRAPHS_GRAPHICS 0
 #endif
 
 #if D3D12_MAX_DEVICE_INTERFACE >= 12
diff --git a/Engine/Source/Runtime/D3D12RHI/Private/Windows/WindowsD3D12Device.cpp b/Engine/Source/Runtime/D3D12RHI/Private/Windows/WindowsD3D12Device.cpp
--- a/Engine/Source/Runtime/D3D12RHI/Private/Windows/WindowsD3D12Device.cpp
+++ b/Engine/Source/Runtime/D3D12RHI/Private/Windows/WindowsD3D12Device.cpp
@@ -1719,11 +1719,12 @@ void FD3D12DynamicRHI::Init()
 
 	GRHISupportsTextureStreaming = true;
 	GRHISupportsFirstInstance = true;
-
 	GRHISupportsShaderRootConstants = true;
-	GRHISupportsShaderBundleDispatch = false; // Shader Bundles only implemented with Work Graphs
-	GRHISupportsShaderBundleWorkGraphDispatch = GRHISupportsShaderWorkGraphsTier1;
-	GRHISupportsShaderBundleParallel = false; // TODO: FD3D12ExplicitDescriptorHeap::UpdateSyncPoint() is not safe for parallel translate due to frame fencing
+
+	GRHIGlobals.ShaderBundles.SupportsDispatch = false; // Shader Bundles only implemented with Work Graphs
+	GRHIGlobals.ShaderBundles.SupportsWorkGraphDispatch = GRHIGlobals.SupportsShaderWorkGraphsTier1;
+	GRHIGlobals.ShaderBundles.SupportsWorkGraphGraphicsDispatch = GRHIGlobals.SupportsShaderWorkGraphsTier1_1;
+	GRHIGlobals.ShaderBundles.SupportsParallel = false; // TODO: FD3D12ExplicitDescriptorHeap::UpdateSyncPoint() is not safe for parallel translate due to frame fencing
 
 	// PC D3D12 support async dispatch rays calls
 	if (GRHIGlobals.RayTracing.Supported)
diff --git a/Engine/Source/Runtime/RHI/Private/DataDrivenShaderPlatformInfo.cpp b/Engine/Source/Runtime/RHI/Private/DataDrivenShaderPlatformInfo.cpp
--- a/Engine/Source/Runtime/RHI/Private/DataDrivenShaderPlatformInfo.cpp
+++ b/Engine/Source/Runtime/RHI/Private/DataDrivenShaderPlatformInfo.cpp
@@ -345,6 +345,7 @@ void FGenericDataDrivenShaderPlatformInfo::ParseDataDrivenShaderInfo(const FConf
 	GET_SECTION_BOOL_HELPER(bSupportsWave64);
 	GET_SECTION_BOOL_HELPER(bSupportsIndependentSamplers);
 	GET_SECTION_BOOL_HELPER(bSupportsWorkGraphs);
+	GET_SECTION_BOOL_HELPER(bSupportsWorkGraphsTier1_1);
 #if WITH_EDITOR
 	GET_SECTION_BOOL_HELPER(bCanUsePreviewPlatformForMaterialValidation);
 #endif
diff --git a/Engine/Source/Runtime/RHI/Public/DataDrivenShaderPlatformInfo.h b/Engine/Source/Runtime/RHI/Public/DataDrivenShaderPlatformInfo.h
--- a/Engine/Source/Runtime/RHI/Public/DataDrivenShaderPlatformInfo.h
+++ b/Engine/Source/Runtime/RHI/Public/DataDrivenShaderPlatformInfo.h
@@ -128,6 +128,7 @@ class FGenericDataDrivenShaderPlatformInfo
 	uint32 bSupportsWave64 : 1;
 	uint32 bSupportsIndependentSamplers : 1;
 	uint32 bSupportsWorkGraphs : 1;
+	uint32 bSupportsWorkGraphsTier1_1 : 1;
 #if WITH_EDITOR
 	uint32 bCanUsePreviewPlatformForMaterialValidation : 1;
 #endif
@@ -833,6 +834,12 @@ public:
 		return Infos[Platform].bSupportsWorkGraphs;
 	}
 
+	static FORCEINLINE_DEBUGGABLE const bool GetSupportsWorkGraphsTier1_1(const FStaticShaderPlatform Platform)
+	{
+		check(IsValid(Platform));
+		return Infos[Platform].bSupportsWorkGraphsTier1_1;
+	}
+
 #if WITH_EDITOR
 	static FORCEINLINE_DEBUGGABLE const bool CanUseForMaterialValidation(const FStaticShaderPlatform Platform)
 	{
@@ -1243,3 +1250,8 @@ inline bool RHISupportsWorkGraphs(const FStaticShaderPlatform Platform)
 {
 	return FDataDrivenShaderPlatformInfo::GetSupportsWorkGraphs(Platform);
 }
+
+inline bool RHISupportsWorkGraphsTier1_1(const FStaticShaderPlatform Platform)
+{
+	return FDataDrivenShaderPlatformInfo::GetSupportsWorkGraphsTier1_1(Platform);
+}
diff --git a/Engine/Source/Runtime/RHI/Public/RHIGlobals.h b/Engine/Source/Runtime/RHI/Public/RHIGlobals.h
--- a/Engine/Source/Runtime/RHI/Public/RHIGlobals.h
+++ b/Engine/Source/Runtime/RHI/Public/RHIGlobals.h
@@ -595,6 +595,9 @@ struct FRHIGlobals
 		/** Whether current RHI supports shader bundle dispatch using work graphs. */
 		bool SupportsWorkGraphDispatch = false;
 
+		/** Whether current RHI supports shader bundle graphics dispatch using work graphs. */
+		bool SupportsWorkGraphGraphicsDispatch = false;
+
 		/** Whether current RHI supports shader bundle dispatch and RHI parallel translate. */
 		bool SupportsParallel = false;
 
@@ -609,9 +612,12 @@ struct FRHIGlobals
 	/** true if the RHI supports Mesh and Amplification shaders with tier1 capability */
 	bool SupportsMeshShadersTier1 = false;
 
-	/** Whether current RHI supports work graphs with tier1 capability. */
+	/** Whether current RHI supports work graphs with tier1_0 capability. */
 	bool SupportsShaderWorkGraphsTier1 = false;
 
+	/** Whether current RHI supports work graphs with tier1_1 capability. */
+	bool SupportsShaderWorkGraphsTier1_1 = false;
+
 	/**
 	* True if the RHI supports reading system timer in shaders via GetShaderTimestamp().
 	* Individual shaders must be compiled with an appropriate vendor extension and check PLATFORM_SUPPORTS_SHADER_TIMESTAMP.
