Add DispatchWorkGraphShaderBundle() stub for graphics bundles.

diff --git a/Engine/Source/Runtime/D3D12RHI/Private/D3D12CommandContext.h b/Engine/Source/Runtime/D3D12RHI/Private/D3D12CommandContext.h
--- a/Engine/Source/Runtime/D3D12RHI/Private/D3D12CommandContext.h
+++ b/Engine/Source/Runtime/D3D12RHI/Private/D3D12CommandContext.h
@@ -817,6 +817,7 @@ private:
 	static void ClearUAV(TRHICommandList_RecursiveHazardous<FD3D12CommandContext>& RHICmdList, FD3D12UnorderedAccessView_RHI* UAV, const void* ClearValues, bool bFloat);
 
 	void DispatchWorkGraphShaderBundle(FRHIShaderBundle* ShaderBundle, FRHIBuffer* RecordArgBuffer, TConstArrayView<FRHIShaderParameterResource> SharedBindlessParameters, TConstArrayView<FRHIShaderBundleComputeDispatch> Dispatches);
+	void DispatchWorkGraphShaderBundle(FRHIShaderBundle* ShaderBundle, FRHIBuffer* RecordArgBuffer, const FRHIShaderBundleGraphicsState& BundleState, TConstArrayView<FRHIShaderParameterResource> SharedBindlessParameters, TConstArrayView<FRHIShaderBundleGraphicsDispatch> Dispatches);
 
 	void HandleDiscardResources          (TArrayView<const FRHITransition*> Transitions, bool bIsBeginTransition);
 	void HandleResourceTransitions       (const struct FD3D12TransitionData* TransitionData, bool& bUAVBarrier);
diff --git a/Engine/Source/Runtime/D3D12RHI/Private/D3D12Commands.cpp b/Engine/Source/Runtime/D3D12RHI/Private/D3D12Commands.cpp
--- a/Engine/Source/Runtime/D3D12RHI/Private/D3D12Commands.cpp
+++ b/Engine/Source/Runtime/D3D12RHI/Private/D3D12Commands.cpp
@@ -1660,8 +1660,7 @@ void FD3D12CommandContext::RHIDispatchGraphicsShaderBundle(
 	}
 	else
 	{
-		checkNoEntry(); // Not implemented
-		//DispatchWorkGraphShaderBundle(ShaderBundle, RecordArgBuffer, Dispatches);
+		DispatchWorkGraphShaderBundle(ShaderBundle, RecordArgBuffer, BundleState, SharedBindlessParameters, Dispatches);
 	}
 }
 
diff --git a/Engine/Source/Runtime/D3D12RHI/Private/D3D12WorkGraph.cpp b/Engine/Source/Runtime/D3D12RHI/Private/D3D12WorkGraph.cpp
--- a/Engine/Source/Runtime/D3D12RHI/Private/D3D12WorkGraph.cpp
+++ b/Engine/Source/Runtime/D3D12RHI/Private/D3D12WorkGraph.cpp
@@ -909,3 +909,13 @@ void FD3D12CommandContext::DispatchWorkGraphShaderBundle(
 
 #endif // D3D12_RHI_WORKGRAPHS
 }
+
+void FD3D12CommandContext::DispatchWorkGraphShaderBundle(
+	FRHIShaderBundle* ShaderBundle, 
+	FRHIBuffer* RecordArgBuffer, 
+	const FRHIShaderBundleGraphicsState& BundleState, 
+	TConstArrayView<FRHIShaderParameterResource> SharedBindlessParameters, 
+	TConstArrayView<FRHIShaderBundleGraphicsDispatch> Dispatches)
+{
+	checkNoEntry(); // Not implemented
+}
diff --git a/Engine/Source/Runtime/RHI/Public/RHICommandListCommandExecutes.inl b/Engine/Source/Runtime/RHI/Public/RHICommandListCommandExecutes.inl
--- a/Engine/Source/Runtime/RHI/Public/RHICommandListCommandExecutes.inl
+++ b/Engine/Source/Runtime/RHI/Public/RHICommandListCommandExecutes.inl
@@ -259,7 +259,7 @@ void FRHICommandDispatchGraphicsShaderBundle::Execute(FRHICommandListBase& CmdLi
 	for (int32 DispatchIndex = 0; DispatchIndex < Dispatches.Num(); ++DispatchIndex)
 	{
 		FRHIShaderBundleGraphicsDispatch& Dispatch = Dispatches[DispatchIndex];
-		if (Dispatch.RecordIndex != ~uint32(0u))
+		if (Dispatch.RecordIndex != ~uint32(0u) && Dispatch.PipelineState)
 		{
 			Dispatch.RHIPipeline = ExecuteSetGraphicsPipelineState(Dispatch.PipelineState);
 		}
