Add listener support for texture replace on d3d12 texture reference

diff --git a/Engine/Source/Runtime/D3D12RHI/Private/D3D12TextureReference.cpp b/Engine/Source/Runtime/D3D12RHI/Private/D3D12TextureReference.cpp
--- a/Engine/Source/Runtime/D3D12RHI/Private/D3D12TextureReference.cpp
+++ b/Engine/Source/Runtime/D3D12RHI/Private/D3D12TextureReference.cpp
@@ -25,6 +25,8 @@ FD3D12RHITextureReference::FD3D12RHITextureReference(FD3D12Device* InDevice, FD3
 
 FD3D12RHITextureReference::~FD3D12RHITextureReference()
 {
+	check(!HasListeners());
+
 #if PLATFORM_SUPPORTS_BINDLESS_RENDERING
 	if (BindlessHandle.IsValid())
 	{
@@ -41,12 +43,16 @@ FD3D12RHITextureReference::~FD3D12RHITextureReference()
 
 void FD3D12RHITextureReference::SwitchToNewTexture(const FD3D12ContextArray& Contexts, FD3D12Texture* InNewTexture)
 {
+	FD3D12Texture* CurrentTexture = FD3D12DynamicRHI::ResourceCast(GetReferencedTexture());
+	FD3D12Texture* NewTexture = InNewTexture ? InNewTexture : FD3D12DynamicRHI::ResourceCast(FRHITextureReference::GetDefaultTexture());
+	if (CurrentTexture != NewTexture)
+	{
+		NotifyListeners(Contexts, CurrentTexture, NewTexture);
+	}
+
 #if PLATFORM_SUPPORTS_BINDLESS_RENDERING
 	if (BindlessHandle.IsValid())
 	{
-		FD3D12Texture* CurrentTexture = FD3D12DynamicRHI::ResourceCast(GetReferencedTexture());
-		FD3D12Texture* NewTexture = InNewTexture ? InNewTexture : FD3D12DynamicRHI::ResourceCast(FRHITextureReference::GetDefaultTexture());
-
 		FD3D12ShaderResourceView* NewView = NewTexture->GetShaderResourceView();
 		const FD3D12OfflineDescriptor NewViewDescriptor = NewView->GetOfflineCpuHandle();
 
diff --git a/Engine/Source/Runtime/D3D12RHI/Private/D3D12TextureReference.h b/Engine/Source/Runtime/D3D12RHI/Private/D3D12TextureReference.h
--- a/Engine/Source/Runtime/D3D12RHI/Private/D3D12TextureReference.h
+++ b/Engine/Source/Runtime/D3D12RHI/Private/D3D12TextureReference.h
@@ -5,6 +5,14 @@
 #include "D3D12RHIPrivate.h"
 #include "RHITextureReference.h"
 
+class FD3D12RHITextureReference;
+
+/** Resource which might needs to be notified about changes on replaced referenced textures */
+struct FD3D12TextureReferenceReplaceListener
+{
+	virtual void TextureReplaced(FD3D12ContextArray const& Contexts, FD3D12RHITextureReference* TextureReference, FD3D12Texture* CurrentTexture, FD3D12Texture* NewTexture) = 0;
+};
+
 class FD3D12RHITextureReference : public FD3D12DeviceChild, public FRHITextureReference, public FD3D12LinkedAdapterObject<FD3D12RHITextureReference>
 #if PLATFORM_SUPPORTS_BINDLESS_RENDERING
 	, public FD3D12ShaderResourceRenameListener
@@ -21,10 +29,45 @@ public:
 #endif // PLATFORM_SUPPORTS_BINDLESS_RENDERING
 
 	void SwitchToNewTexture(FD3D12ContextArray const& Contexts, FD3D12Texture* InNewTexture);
-
+		
 #if PLATFORM_SUPPORTS_BINDLESS_RENDERING
 	uint32 ReferencedDescriptorVersion{};
 #endif
+
+	void AddReplaceListener(FD3D12TextureReferenceReplaceListener* InListener)
+	{
+		FScopeLock Lock(&TextureReplaceListenersCS);
+		check(!TextureReplaceListeners.Contains(InListener));
+		TextureReplaceListeners.Add(InListener);
+	}
+
+	void RemoveReplaceListener(FD3D12TextureReferenceReplaceListener* InListener)
+	{
+		FScopeLock Lock(&TextureReplaceListenersCS);
+		uint32 Removed = TextureReplaceListeners.Remove(InListener);
+
+		checkf(Removed == 1, TEXT("Should have exactly one registered listener during remove (same listener shouldn't registered twice and we shouldn't call this if not registered"));
+	}
+
+private:
+
+	bool HasListeners() const
+	{
+		FScopeLock Lock(&TextureReplaceListenersCS);
+		return TextureReplaceListeners.Num() != 0;
+	}
+
+	void NotifyListeners(FD3D12ContextArray const& Contexts, FD3D12Texture* CurrentTexture, FD3D12Texture* NewTexture)
+	{
+		FScopeLock Lock(&TextureReplaceListenersCS);
+		for (FD3D12TextureReferenceReplaceListener* ReplaceListener : TextureReplaceListeners)
+		{
+			ReplaceListener->TextureReplaced(Contexts, this, CurrentTexture, NewTexture);
+		}
+	}
+
+	mutable FCriticalSection TextureReplaceListenersCS;
+	TArray<FD3D12TextureReferenceReplaceListener*> TextureReplaceListeners;
 };
 
 template<>
