Relaxed presentation check with additional logging

diff --git a/Engine/Source/Runtime/D3D12RHI/Private/D3D12Viewport.cpp b/Engine/Source/Runtime/D3D12RHI/Private/D3D12Viewport.cpp
--- a/Engine/Source/Runtime/D3D12RHI/Private/D3D12Viewport.cpp
+++ b/Engine/Source/Runtime/D3D12RHI/Private/D3D12Viewport.cpp
@@ -278,10 +278,21 @@ void FD3D12Viewport::Resize(uint32 InSizeX, uint32 InSizeY, bool bInIsFullscreen
 	if (RHIConsoleVariables::LogViewportEvents)
 	{
 		const FString& ThreadName = FThreadManager::GetThreadName(FPlatformTLS::GetCurrentThreadId());
-		UE_LOG(LogD3D12RHI, Log, TEXT("Thread %s: Resize Viewport %#016llx (%ux%u)"), ThreadName.GetCharArray().GetData(), this, InSizeX, InSizeY);
+		UE_LOG(LogD3D12RHI, Log, TEXT("Thread %s: Resize Viewport, %s"), ThreadName.GetCharArray().GetData(), *GetStateString());
 	}
 #endif
 
+	// Log relevant state changes, makes it easier to track/reproduce crashes
+	const bool bLogEvent =
+		(bIsFullscreen != bInIsFullscreen) ||
+		(bIsFullscreen && (InSizeX != SizeX || InSizeY != SizeY));
+
+	FString OldState;
+	if (bLogEvent)
+	{
+		OldState = GetStateString();
+	}
+
 	// Flush the outstanding GPU work and wait for it to complete.
 	FlushRenderingCommands();
 	Adapter->BlockUntilIdle();
@@ -409,6 +420,11 @@ void FD3D12Viewport::Resize(uint32 InSizeX, uint32 InSizeY, bool bInIsFullscreen
 	{
 		ShutdownHDR();
 	}
+
+	if (bLogEvent)
+	{
+		UE_LOG(LogD3D12RHI, Log, TEXT("Swapchain Resized:\n\tBefore: %s\n\tAfter: %s"), *OldState, *GetStateString());
+	}
 }
 
 #if D3D12_VIEWPORT_EXPOSES_SWAP_CHAIN
@@ -462,6 +478,24 @@ void FD3D12Viewport::AdvanceExpectedBackBufferIndex_RenderThread()
 	}
 }
 
+FString FD3D12Viewport::GetStateString()
+{
+	return FString::Printf(
+		TEXT("Viewport=0x%p, Num=%d, Size=(%d,%d), PF=%d, DXGIFormat=0x%x, Fullscreen=%d, AllowTearing=%d"),
+		this,
+		NumBackBuffers, SizeX, SizeY, static_cast<int32>(PixelFormat),
+		static_cast<int32>(UE::DXGIUtilities::GetSwapChainFormat(PixelFormat)),
+		static_cast<int32>(bIsFullscreen),
+		static_cast<int32>(bAllowTearing)
+	);
+}
+
+static bool IsTransientPresentationError(HRESULT Result)
+{
+	return Result == E_INVALIDARG ||
+		   Result == DXGI_ERROR_INVALID_CALL;
+}
+
 /** Presents the swap chain checking the return result. */
 bool FD3D12Viewport::PresentChecked(IRHICommandContext& RHICmdContext, int32 SyncInterval)
 {
@@ -502,8 +536,31 @@ bool FD3D12Viewport::PresentChecked(IRHICommandContext& RHICmdContext, int32 Syn
 
 	if (bNeedNativePresent)
 	{
+		static constexpr uint32 MaxPresentFailures = 5;
+
 		// Present the back buffer to the viewport window.
-		VERIFYD3D12RESULT(PresentInternal(SyncInterval));
+		// In case presentation failures are transient, don't fault on the first one.
+		HRESULT Result = PresentInternal(SyncInterval);
+		if (SUCCEEDED(Result))
+		{
+			CheckedPresentFailureCounter = 0;
+		}
+		else if (!IsTransientPresentationError(Result) || ++CheckedPresentFailureCounter >= MaxPresentFailures)
+		{			
+			VERIFYD3D12RESULT_LAMBDA(Result, GetParentAdapter()->GetD3DDevice(), [this]
+			{
+				return GetStateString();
+			});
+		}
+		else
+		{
+			UE_LOG(
+				LogD3D12RHI, Error,
+				TEXT("Swapchain presentation try %u/%u failed with HR(0x%x): %s"),
+				CheckedPresentFailureCounter, MaxPresentFailures,
+				Result, *GetStateString()
+			);
+		}
 
 		if (IsValidRef(CustomPresent))
 		{
@@ -729,7 +786,11 @@ void FD3D12CommandContextBase::RHIEndDrawingViewport(FRHIViewport* ViewportRHI,
 	check(ParentAdapter->GetDrawingViewport() == Viewport);
 	ParentAdapter->SetDrawingViewport(nullptr);
 
-	bool bNativelyPresented = Viewport->Present(*this, bLockToVsync);
+	bool bNativelyPresented = true;
+	if (bPresent)
+	{
+		bNativelyPresented = Viewport->Present(*this, bLockToVsync);
+	}
 	
 	// Multi-GPU support : here each GPU wait's for it's own frame completion.
 	if (bNativelyPresented)
diff --git a/Engine/Source/Runtime/D3D12RHI/Private/D3D12Viewport.h b/Engine/Source/Runtime/D3D12RHI/Private/D3D12Viewport.h
--- a/Engine/Source/Runtime/D3D12RHI/Private/D3D12Viewport.h
+++ b/Engine/Source/Runtime/D3D12RHI/Private/D3D12Viewport.h
@@ -148,6 +148,8 @@ public:
 private:
 	bool IsPresentAllowed();
 
+	FString GetStateString();
+
 #if D3D12RHI_USE_DUMMY_BACKBUFFER
 	FD3D12Texture* CreateDummyBackBufferTextures(FD3D12Adapter* InAdapter, EPixelFormat InPixelFormat, uint32 InSizeX, uint32 InSizeY);
 #endif
@@ -212,6 +214,8 @@ private:
 	bool bAllowTearing = true;
 	bool bNeedSwapChain = false;
 
+	uint32 CheckedPresentFailureCounter = 0;
+
 #if D3D12_VIEWPORT_EXPOSES_SWAP_CHAIN
 	TRefCountPtr<IDXGISwapChain1> SwapChain1;
 #if DXGI_MAX_SWAPCHAIN_INTERFACE >= 2
diff --git a/Engine/Source/Runtime/D3D12RHI/Private/Windows/WindowsD3D12Viewport.cpp b/Engine/Source/Runtime/D3D12RHI/Private/Windows/WindowsD3D12Viewport.cpp
--- a/Engine/Source/Runtime/D3D12RHI/Private/Windows/WindowsD3D12Viewport.cpp
+++ b/Engine/Source/Runtime/D3D12RHI/Private/Windows/WindowsD3D12Viewport.cpp
@@ -304,12 +304,10 @@ void FD3D12Viewport::ResizeInternal()
 	{
 		if (SwapChain1)
 		{
-			auto Lambda = [this, SwapChainFlags]() -> FString
+			VERIFYD3D12RESULT_LAMBDA(SwapChain1->ResizeBuffers(NumBackBuffers, SizeX, SizeY, UE::DXGIUtilities::GetSwapChainFormat(PixelFormat), SwapChainFlags), Adapter->GetD3DDevice(), [this]
 			{
-				return FString::Printf(TEXT("Num=%d, Size=(%d,%d), PF=%d, DXGIFormat=0x%x, Flags=0x%x"), NumBackBuffers, SizeX, SizeY, (int32)PixelFormat, (int32)UE::DXGIUtilities::GetSwapChainFormat(PixelFormat), SwapChainFlags);
-			};
-
-			VERIFYD3D12RESULT_LAMBDA(SwapChain1->ResizeBuffers(NumBackBuffers, SizeX, SizeY, UE::DXGIUtilities::GetSwapChainFormat(PixelFormat), SwapChainFlags), Adapter->GetD3DDevice(), Lambda);
+				return GetStateString();
+			});
 		}
 
 		FD3D12Device* Device = Adapter->GetDevice(0);
