Workgraph shader bundle node naming fix.
Empty shader slots should still increment node count so that naming naming of workgraph nodes matches shader bundle bins.

diff --git a/Engine/Source/Runtime/D3D12RHI/Private/D3D12WorkGraph.cpp b/Engine/Source/Runtime/D3D12RHI/Private/D3D12WorkGraph.cpp
--- a/Engine/Source/Runtime/D3D12RHI/Private/D3D12WorkGraph.cpp
+++ b/Engine/Source/Runtime/D3D12RHI/Private/D3D12WorkGraph.cpp
@@ -71,19 +71,20 @@ FD3D12WorkGraphPipelineState::FD3D12WorkGraphPipelineState(FD3D12Device* Device,
 	for (int32 Index = 0; Index < ShaderTableNum; ++Index) 
 	{
 		FD3D12WorkGraphShader* NodeShader = (FD3D12WorkGraphShader*)Initializer.GetShaderTable()[Index];
+		static const FString EmptyString;
+		FString const& ExportName = NodeShader == nullptr ? EmptyString : NodeShader->EntryPoint;
+		const int32 NodeNameIndex = Initializer.GetNameTable().IndexOfByPredicate([ExportName](FWorkGraphPipelineStateInitializer::FNameMap const& NameMap)
+		{
+			return NameMap.ExportName == ExportName;
+		});
 
-		if (NodeShader != nullptr)
+		if (NodeNameIndex != INDEX_NONE)
 		{
-			const int32 NodeNameIndex = Initializer.GetNameTable().IndexOfByPredicate([NodeShader](FWorkGraphPipelineStateInitializer::FNameMap const& NameMap)
-			{
-				return NameMap.ExportName == NodeShader->EntryPoint;
-			});
+			FString const& NodeName = Initializer.GetNameTable()[NodeNameIndex].NodeName;
+			const uint32 NodeArrayIndex = (NodeCountPerName.FindOrAdd(NodeName))++;
 
-			if (NodeNameIndex != INDEX_NONE)
+			if (NodeShader != nullptr)
 			{
-				FString const& NodeName = Initializer.GetNameTable()[NodeNameIndex].NodeName;
-				const uint32 NodeCount = (NodeCountPerName.FindOrAdd(NodeName))++;
-				const uint32 NodeArrayIndex = NodeCount;
 				const FString NodePathName = FString::Printf(TEXT("%s_%d"), *NodeName, NodeArrayIndex);
 
 				CD3DX12_DXIL_LIBRARY_SUBOBJECT* Lib = StateObjectDesc.CreateSubobject<CD3DX12_DXIL_LIBRARY_SUBOBJECT>();
@@ -103,9 +104,9 @@ FD3D12WorkGraphPipelineState::FD3D12WorkGraphPipelineState(FD3D12Device* Device,
 				const int32 LocalRootArgumentsTableIndex = RootArgOffsets.Num();
 				RootArgOffsets.Add(LocalRootArgumentsTableIndex);
 				Override->LocalRootArgumentsTableIndex(LocalRootArgumentsTableIndex);
+		
+				RootArgStrideInBytes = FMath::Max(RootArgStrideInBytes, NodeShader->RootSignature->GetTotalRootSignatureSizeInBytes());
 			}
-
-			RootArgStrideInBytes = FMath::Max(RootArgStrideInBytes, NodeShader->RootSignature->GetTotalRootSignatureSizeInBytes());
 		}
 	}
 	
@@ -689,9 +690,10 @@ void FD3D12CommandContext::DispatchWorkGraphShaderBundle(
 	FWorkGraphPipelineStateInitializer Initializer;
 	Initializer.SetProgramName(TEXT("ShaderBundle"));
 	TArray<FWorkGraphPipelineStateInitializer::FNameMap> NameTable;
-	NameTable.Add(FWorkGraphPipelineStateInitializer::FNameMap(TEXT("WorkGraphMainCS"), TEXT("WorkGraphMainCS")));
-	NameTable.Add(FWorkGraphPipelineStateInitializer::FNameMap(TEXT("MainCS"), TEXT("ShaderBundleNode")));
-	NameTable.Add(FWorkGraphPipelineStateInitializer::FNameMap(TEXT("MicropolyRasterize"), TEXT("ShaderBundleNode")));
+	NameTable.Add(FWorkGraphPipelineStateInitializer::FNameMap(TEXT("WorkGraphMainCS"), TEXT("WorkGraphMainCS"))); // Entry node.
+	NameTable.Add(FWorkGraphPipelineStateInitializer::FNameMap(TEXT(""), TEXT("ShaderBundleNode"))); // Empty shader slots still increase bundle node index.
+	NameTable.Add(FWorkGraphPipelineStateInitializer::FNameMap(TEXT("MainCS"), TEXT("ShaderBundleNode"))); // Nanite compute materials.
+	NameTable.Add(FWorkGraphPipelineStateInitializer::FNameMap(TEXT("MicropolyRasterize"), TEXT("ShaderBundleNode"))); // Nanite software rasterize.
 	Initializer.SetNameTable(NameTable);
 	Initializer.SetShaderTable(LocalNodeShaders);
 
