Work graph implementation for graphics shader bundles.
Currently this is only enabled with the preview agility SDK and shader compiler.
There are a couple of todos that are incoming soon. One is that we are using a hardcoded pixelshader entrypoint until we serialize that in pixel shaders. The other is that we are aliasing the existing mesh shader slot in the PSO bound shaders instead of creating a new one or changing the existing one to the shared base class.

diff --git a/Engine/Source/Runtime/D3D12RHI/Private/D3D12WorkGraph.cpp b/Engine/Source/Runtime/D3D12RHI/Private/D3D12WorkGraph.cpp
--- a/Engine/Source/Runtime/D3D12RHI/Private/D3D12WorkGraph.cpp
+++ b/Engine/Source/Runtime/D3D12RHI/Private/D3D12WorkGraph.cpp
@@ -37,7 +37,7 @@ FD3D12WorkGraphPipelineState::FD3D12WorkGraphPipelineState(FD3D12Device* Device,
 			FRHIWorkGraphShader* Shader = Initializer.GetShaderTable()[RootShaderIndex];
 			if (Shader != nullptr && Shader->GetFrequency() == SF_WorkGraphRoot)
 			{
-				GlobalRootSignature = ((FD3D12WorkGraphShader*)Shader)->RootSignature->GetRootSignature();
+				GlobalRootSignature = FD3D12DynamicRHI::ResourceCast(Shader)->RootSignature->GetRootSignature();
 			}
 		}
 	}
@@ -60,7 +60,7 @@ FD3D12WorkGraphPipelineState::FD3D12WorkGraphPipelineState(FD3D12Device* Device,
 
 	D3D12_NODE_ID EntryPoint;
 	{
-		FD3D12WorkGraphShader* EntryShader = (FD3D12WorkGraphShader*)Initializer.GetShaderTable()[0];
+		FD3D12WorkGraphShader* EntryShader = FD3D12DynamicRHI::ResourceCast(Initializer.GetShaderTable()[0]);
 		EntryPoint.Name = *EntryShader->EntryPoint;
 		EntryPoint.ArrayIndex = 0;
 	}
@@ -70,7 +70,7 @@ FD3D12WorkGraphPipelineState::FD3D12WorkGraphPipelineState(FD3D12Device* Device,
 	const int32 ShaderTableNum = Initializer.GetShaderTable().Num();
 	for (int32 Index = 0; Index < ShaderTableNum; ++Index) 
 	{
-		FD3D12WorkGraphShader* NodeShader = (FD3D12WorkGraphShader*)Initializer.GetShaderTable()[Index];
+		FD3D12WorkGraphShader* NodeShader = FD3D12DynamicRHI::ResourceCast(Initializer.GetShaderTable()[Index]);
 		static const FString EmptyString;
 		FString const& ExportName = NodeShader == nullptr ? EmptyString : NodeShader->EntryPoint;
 		const int32 NodeNameIndex = Initializer.GetNameTable().IndexOfByPredicate([ExportName](FWorkGraphPipelineStateInitializer::FNameMap const& NameMap)
@@ -110,6 +110,139 @@ FD3D12WorkGraphPipelineState::FD3D12WorkGraphPipelineState(FD3D12Device* Device,
 		}
 	}
 	
+#if D3D12_RHI_WORKGRAPHS_GRAPHICS
+
+	// Graphics Shader Table
+	const int32 PSOTableNum = Initializer.GetGraphicsPSOTable().Num();
+	if (PSOTableNum > 0)
+	{
+		CD3DX12_STATE_OBJECT_CONFIG_SUBOBJECT* ConfigSubobject = StateObjectDesc.CreateSubobject<CD3DX12_STATE_OBJECT_CONFIG_SUBOBJECT>();
+		ConfigSubobject->SetFlags(D3D12_STATE_OBJECT_FLAG_WORK_GRAPHS_USE_GRAPHICS_STATE_FOR_GLOBAL_ROOT_SIGNATURE);
+	}
+
+	for (int32 Index = 0; Index < PSOTableNum; ++Index)
+	{
+		FGraphicsPipelineStateInitializer const* NodePSO = Initializer.GetGraphicsPSOTable()[Index];
+		FRHIWorkGraphShader* MeshShader = NodePSO == nullptr ? nullptr : (FRHIWorkGraphShader*)NodePSO->BoundShaderState.GetMeshShader(); // todo: use correct types (or base type) in BoundShaderState.
+		FRHIPixelShader* PixelShader = NodePSO == nullptr ? nullptr : NodePSO->BoundShaderState.GetPixelShader();
+
+		static const FString EmptyString;
+		FString const& ExportName = MeshShader == nullptr ? EmptyString : FD3D12DynamicRHI::ResourceCast(MeshShader)->EntryPoint;
+		const int32 NodeNameIndex = Initializer.GetNameTable().IndexOfByPredicate([ExportName](FWorkGraphPipelineStateInitializer::FNameMap const& NameMap)
+		{
+			return NameMap.ExportName == ExportName;
+		});
+
+		if (NodeNameIndex == INDEX_NONE)
+		{
+			continue;
+		}
+
+		FString const& NodeName = Initializer.GetNameTable()[NodeNameIndex].NodeName;
+		const uint32 NodeArrayIndex = (NodeCountPerName.FindOrAdd(NodeName))++;
+
+		if (MeshShader == nullptr || PixelShader == nullptr)
+		{
+			continue;
+		}
+		
+		const FString NodePathName = FString::Printf(TEXT("%s_%d"), *NodeName, NodeArrayIndex);
+		FString MeshShaderName = FString::Printf(TEXT("MeshShader_%d"), NodeArrayIndex);
+		FString PixelShaderName = FString::Printf(TEXT("PixelShader_%d"), NodeArrayIndex);
+			
+		FD3D12RootSignature const* LocalRootSignature = Device->GetParentAdapter()->GetWorkGraphGraphicsRootSignature(NodePSO->BoundShaderState);
+		CD3DX12_LOCAL_ROOT_SIGNATURE_SUBOBJECT* LocalRootSignatureSubobject = StateObjectDesc.CreateSubobject<CD3DX12_LOCAL_ROOT_SIGNATURE_SUBOBJECT>();
+		LocalRootSignatureSubobject->SetRootSignature(LocalRootSignature->GetRootSignature());
+		RootArgStrideInBytes = FMath::Max(RootArgStrideInBytes, LocalRootSignature->GetTotalRootSignatureSizeInBytes());
+
+		{
+			FD3D12WorkGraphShader* D3D12Shader = FD3D12DynamicRHI::ResourceCast(MeshShader);
+
+			CD3DX12_DXIL_LIBRARY_SUBOBJECT* LibrarySubobject = StateObjectDesc.CreateSubobject<CD3DX12_DXIL_LIBRARY_SUBOBJECT>();
+			CD3DX12_SHADER_BYTECODE ByteCode(D3D12Shader->Code.GetData(), D3D12Shader->Code.Num());
+			LibrarySubobject->SetDXILLibrary(&ByteCode);
+			LibrarySubobject->DefineExport(*MeshShaderName, *D3D12Shader->EntryPoint);
+
+			CD3DX12_SUBOBJECT_TO_EXPORTS_ASSOCIATION_SUBOBJECT* ExportAssociationSubobject = StateObjectDesc.CreateSubobject<CD3DX12_SUBOBJECT_TO_EXPORTS_ASSOCIATION_SUBOBJECT>();
+			ExportAssociationSubobject->SetSubobjectToAssociate(*LocalRootSignatureSubobject);
+			ExportAssociationSubobject->AddExport(*MeshShaderName);
+		}
+		{
+			FD3D12PixelShader* D3D12Shader = FD3D12DynamicRHI::ResourceCast(PixelShader);
+
+			CD3DX12_DXIL_LIBRARY_SUBOBJECT* LibrarySubobject = StateObjectDesc.CreateSubobject<CD3DX12_DXIL_LIBRARY_SUBOBJECT>();
+			CD3DX12_SHADER_BYTECODE ByteCode(D3D12Shader->Code.GetData(), D3D12Shader->Code.Num());
+			LibrarySubobject->SetDXILLibrary(&ByteCode);
+			LibrarySubobject->DefineExport(*PixelShaderName, TEXT("HWRasterizePS")); // todo: Store EntryPoint in pixel shader and use that. Using hard coded value for only use case here!
+
+			CD3DX12_SUBOBJECT_TO_EXPORTS_ASSOCIATION_SUBOBJECT* ExportAssociationSubobject = StateObjectDesc.CreateSubobject<CD3DX12_SUBOBJECT_TO_EXPORTS_ASSOCIATION_SUBOBJECT>();
+			ExportAssociationSubobject->SetSubobjectToAssociate(*LocalRootSignatureSubobject);
+			ExportAssociationSubobject->AddExport(*PixelShaderName);
+		}
+
+		TArray<CD3DX12_STATE_OBJECT_DESC::SUBOBJECT_HELPER_BASE*, TInlineAllocator<8>> StateSubobjects;
+
+		CD3DX12_PRIMITIVE_TOPOLOGY_SUBOBJECT* PrimitiveTopologyState = StateObjectDesc.CreateSubobject<CD3DX12_PRIMITIVE_TOPOLOGY_SUBOBJECT>();
+		const EPrimitiveType PrimitiveType = NodePSO->PrimitiveType;
+		PrimitiveTopologyState->SetPrimitiveTopologyType(D3D12PrimitiveTypeToTopologyType(TranslatePrimitiveType(PrimitiveType)));
+
+		FD3D12RasterizerState* D3D12RasterizerState = FD3D12DynamicRHI::ResourceCast(NodePSO->RasterizerState);
+		CD3DX12_RASTERIZER_SUBOBJECT* RasterizerSubobject = StateObjectDesc.CreateSubobject<CD3DX12_RASTERIZER_SUBOBJECT>();
+		RasterizerSubobject->SetFrontCounterClockwise(D3D12RasterizerState->Desc.FrontCounterClockwise);
+		RasterizerSubobject->SetFillMode(D3D12RasterizerState->Desc.FillMode);
+		RasterizerSubobject->SetCullMode(D3D12RasterizerState->Desc.CullMode);
+		StateSubobjects.Add(RasterizerSubobject);
+
+		FD3D12DepthStencilState* D3D12DepthStencilState = FD3D12DynamicRHI::ResourceCast(NodePSO->DepthStencilState);
+		CD3DX12_DEPTH_STENCIL_SUBOBJECT* DepthStencilSubobject = StateObjectDesc.CreateSubobject<CD3DX12_DEPTH_STENCIL_SUBOBJECT>();
+		DepthStencilSubobject->SetDepthEnable(D3D12DepthStencilState->Desc.DepthEnable);
+		DepthStencilSubobject->SetDepthFunc(D3D12DepthStencilState->Desc.DepthFunc);
+		DepthStencilSubobject->SetDepthWriteMask(D3D12DepthStencilState->Desc.DepthWriteMask);
+		DepthStencilSubobject->SetStencilEnable(D3D12DepthStencilState->Desc.StencilEnable);
+		DepthStencilSubobject->SetStencilReadMask(D3D12DepthStencilState->Desc.StencilReadMask);
+		DepthStencilSubobject->SetStencilWriteMask(D3D12DepthStencilState->Desc.StencilWriteMask);
+		StateSubobjects.Add(DepthStencilSubobject);
+
+		if (NodePSO->DepthStencilTargetFormat != PF_Unknown)
+		{
+			CD3DX12_DEPTH_STENCIL_FORMAT_SUBOBJECT* DepthStencilFormatSubobject = StateObjectDesc.CreateSubobject<CD3DX12_DEPTH_STENCIL_FORMAT_SUBOBJECT>();
+			DepthStencilFormatSubobject->SetDepthStencilFormat((DXGI_FORMAT)GPixelFormats[NodePSO->DepthStencilTargetFormat].PlatformFormat);
+			StateSubobjects.Add(DepthStencilFormatSubobject);
+		}
+
+		CD3DX12_RENDER_TARGET_FORMATS_SUBOBJECT* RenderTargetFormatSubobject = StateObjectDesc.CreateSubobject<CD3DX12_RENDER_TARGET_FORMATS_SUBOBJECT>();
+		const int32 NumRenderTargets = NodePSO->ComputeNumValidRenderTargets();
+		RenderTargetFormatSubobject->SetNumRenderTargets(NumRenderTargets);
+		if (NumRenderTargets > 0)
+		{
+			for (int32 RenderTargetIndex = 0; RenderTargetIndex < NumRenderTargets; ++RenderTargetIndex)
+			{
+				RenderTargetFormatSubobject->SetRenderTargetFormat(RenderTargetIndex, (DXGI_FORMAT)GPixelFormats[NodePSO->RenderTargetFormats[RenderTargetIndex]].PlatformFormat);
+			}
+		}
+		StateSubobjects.Add(RenderTargetFormatSubobject);
+
+		CD3DX12_GENERIC_PROGRAM_SUBOBJECT* ProgramSubobject = StateObjectDesc.CreateSubobject<CD3DX12_GENERIC_PROGRAM_SUBOBJECT>();
+		ProgramSubobject->SetProgramName(*NodePathName);
+		ProgramSubobject->AddExport(*MeshShaderName);
+		ProgramSubobject->AddExport(*PixelShaderName);
+		for (CD3DX12_STATE_OBJECT_DESC::SUBOBJECT_HELPER_BASE* StateSubobject : StateSubobjects)
+		{
+			ProgramSubobject->AddSubobject(*StateSubobject);
+		}
+		ProgramSubobject->Finalize();
+
+		CD3DX12_MESH_LAUNCH_NODE_OVERRIDES* NodeOverrides = WorkGraphSubobject->CreateMeshLaunchNodeOverrides(*NodePathName);
+		NodeOverrides->NewName(D3D12_NODE_ID{ *NodeName, NodeArrayIndex });
+		NodeOverrides->MaxInputRecordsPerGraphEntryRecord(1, false);
+
+		const int32 LocalRootArgumentsTableIndex = RootArgOffsets.Num();
+		RootArgOffsets.Add(LocalRootArgumentsTableIndex);
+		NodeOverrides->LocalRootArgumentsTableIndex(LocalRootArgumentsTableIndex);
+	}
+	
+#endif // D3D12_RHI_WORKGRAPHS_GRAPHICS
+
 	RootArgStrideInBytes = ((RootArgStrideInBytes + 15) & ~15);
 	MaxRootArgOffset = RootArgOffsets.Last();
 
@@ -125,11 +258,19 @@ FD3D12WorkGraphPipelineState::FD3D12WorkGraphPipelineState(FD3D12Device* Device,
 
 	ProgramIdentifier = PipelineProperties->GetProgramIdentifier(ProgramName);
 
+#if D3D12_RHI_WORKGRAPHS_GRAPHICS
+	TRefCountPtr<ID3D12WorkGraphProperties1> WorkGraphProperties;
+#else
 	TRefCountPtr<ID3D12WorkGraphProperties> WorkGraphProperties;
+#endif
+
 	HResult = StateObject->QueryInterface(IID_PPV_ARGS(WorkGraphProperties.GetInitReference()));
 	checkf(SUCCEEDED(HResult), TEXT("Failed to query work graph properties from the work graph pipeline state object. Result=%08x"), uint32(HResult));
 
 	UINT WorkGraphIndex = WorkGraphProperties->GetWorkGraphIndex(ProgramName);
+#if D3D12_RHI_WORKGRAPHS_GRAPHICS
+	WorkGraphProperties->SetMaximumInputRecords(WorkGraphIndex, 1, 1);
+#endif
 	D3D12_WORK_GRAPH_MEMORY_REQUIREMENTS MemoryRequirements = {};
 	WorkGraphProperties->GetWorkGraphMemoryRequirements(WorkGraphIndex, &MemoryRequirements);
 
@@ -653,7 +794,7 @@ void FD3D12CommandContext::DispatchWorkGraphShaderBundle(
 	FD3D12ShaderBundle* D3D12ShaderBundle = static_cast<FD3D12ShaderBundle*>(FD3D12DynamicRHI::ResourceCast(ShaderBundle));
 
 	TShaderRef<FDispatchShaderBundleWorkGraph> WorkGraphGlobalShader = GetGlobalShaderMap(GMaxRHIFeatureLevel)->GetShader<FDispatchShaderBundleWorkGraph>();
-	FD3D12WorkGraphShader* WorkGraphGlobalShaderRHI = static_cast<FD3D12WorkGraphShader*>(WorkGraphGlobalShader.GetWorkGraphShader());
+	FD3D12WorkGraphShader* WorkGraphGlobalShaderRHI = FD3D12DynamicRHI::ResourceCast(WorkGraphGlobalShader.GetWorkGraphShader());
 	const bool bBindlessResources = WorkGraphGlobalShaderRHI->UsesBindlessResources();
 
 	uint32 ViewDescriptorCount = WorkGraphGlobalShaderRHI->ResourceCounts.NumSRVs + WorkGraphGlobalShaderRHI->ResourceCounts.NumCBs + WorkGraphGlobalShaderRHI->ResourceCounts.NumUAVs;
@@ -919,5 +1060,330 @@ void FD3D12CommandContext::DispatchWorkGraphShaderBundle(
 	TConstArrayView<FRHIShaderParameterResource> SharedBindlessParameters, 
 	TConstArrayView<FRHIShaderBundleGraphicsDispatch> Dispatches)
 {
-	checkNoEntry(); // Not implemented
+#if D3D12_RHI_WORKGRAPHS
+
+	TRHICommandList_RecursiveHazardous<FD3D12CommandContext> RHICmdList(this);
+
+	FD3D12ShaderBundle* D3D12ShaderBundle = static_cast<FD3D12ShaderBundle*>(FD3D12DynamicRHI::ResourceCast(ShaderBundle));
+
+	TShaderRef<FDispatchShaderBundleWorkGraph> WorkGraphGlobalShader = GetGlobalShaderMap(GMaxRHIFeatureLevel)->GetShader<FDispatchShaderBundleWorkGraph>();
+	FD3D12WorkGraphShader* WorkGraphGlobalShaderRHI = FD3D12DynamicRHI::ResourceCast(WorkGraphGlobalShader.GetWorkGraphShader());
+	const bool bBindlessResources = WorkGraphGlobalShaderRHI->UsesBindlessResources();
+
+	uint32 ViewDescriptorCount = WorkGraphGlobalShaderRHI->ResourceCounts.NumSRVs + WorkGraphGlobalShaderRHI->ResourceCounts.NumCBs + WorkGraphGlobalShaderRHI->ResourceCounts.NumUAVs;
+	uint32 SamplerDescriptorCount = WorkGraphGlobalShaderRHI->ResourceCounts.NumSamplers;
+
+	TArray<FRHIWorkGraphShader*> LocalNodeShaders;
+	LocalNodeShaders.Add(WorkGraphGlobalShaderRHI);
+
+	const int32 NumRecords = Dispatches.Num();
+	checkf(NumRecords <= FDispatchShaderBundleWorkGraph::GetMaxShaderBundleSize(), TEXT("Too many entries in a shader bundle (%d). Try increasing 'r.ShaderBundle.MaxSize'"), NumRecords);
+
+	TArray<uint32> ValidRecords;
+	ValidRecords.Reserve(NumRecords);
+	TArray<FGraphicsPipelineStateInitializer const*> LocalPSOs;
+	LocalPSOs.Reserve(NumRecords);
+
+	for (int32 DispatchIndex = 0; DispatchIndex < NumRecords; ++DispatchIndex)
+	{
+		const FRHIShaderBundleGraphicsDispatch& Dispatch = Dispatches[DispatchIndex];
+		FGraphicsPipelineStateInitializer const* PSO = Dispatch.IsValid() ? &Dispatch.PipelineInitializer : nullptr;
+		LocalPSOs.Add(PSO);
+
+		if (PSO != nullptr)
+		{
+			FRHIWorkGraphShader* MeshShader = (FRHIWorkGraphShader*)PSO->BoundShaderState.GetMeshShader();
+			FRHIPixelShader* PixelShader = PSO->BoundShaderState.GetPixelShader();
+			if (MeshShader != nullptr && PixelShader != nullptr)
+			{
+				ValidRecords.Add(uint32(DispatchIndex));
+
+				if (FD3D12WorkGraphShader* D3D12Shader = FD3D12DynamicRHI::ResourceCast(MeshShader))
+				{
+					ViewDescriptorCount += D3D12Shader->ResourceCounts.NumSRVs + D3D12Shader->ResourceCounts.NumCBs + D3D12Shader->ResourceCounts.NumUAVs;
+					SamplerDescriptorCount += D3D12Shader->ResourceCounts.NumSamplers;
+				}
+				if (FD3D12PixelShader* D3D12Shader = FD3D12DynamicRHI::ResourceCast(PixelShader))
+				{
+					ViewDescriptorCount += D3D12Shader->ResourceCounts.NumSRVs + D3D12Shader->ResourceCounts.NumCBs + D3D12Shader->ResourceCounts.NumUAVs;
+					SamplerDescriptorCount += D3D12Shader->ResourceCounts.NumSamplers;
+				}
+			}
+		}
+	}
+	const int32 NumValidRecords = ValidRecords.Num();
+
+	FWorkGraphPipelineStateInitializer Initializer;
+	Initializer.SetProgramName(TEXT("ShaderBundle"));
+	TArray<FWorkGraphPipelineStateInitializer::FNameMap> NameTable;
+	NameTable.Add(FWorkGraphPipelineStateInitializer::FNameMap(TEXT("WorkGraphMainCS"), TEXT("WorkGraphMainCS"))); // Entry node.
+	NameTable.Add(FWorkGraphPipelineStateInitializer::FNameMap(TEXT(""), TEXT("ShaderBundleNode"))); // Empty shader slots still increase bundle node index.
+	NameTable.Add(FWorkGraphPipelineStateInitializer::FNameMap(TEXT("HWRasterizeMS"), TEXT("ShaderBundleNode"))); // Nanite software rasterize.
+	Initializer.SetNameTable(NameTable);
+	Initializer.SetShaderTable(LocalNodeShaders);
+	Initializer.SetGraphicsPSOTable(LocalPSOs);
+
+	FWorkGraphPipelineState* WorkGraphPipelineState = PipelineStateCache::GetAndOrCreateWorkGraphPipelineState(RHICmdList, Initializer);
+	FD3D12WorkGraphPipelineState* Pipeline = static_cast<FD3D12WorkGraphPipelineState*>(GetRHIWorkGraphPipelineState(WorkGraphPipelineState));
+
+	const uint32 NumViewDescriptors = ViewDescriptorCount;
+	const uint32 NumSamplerDescriptors = SamplerDescriptorCount;
+
+	const uint32 MaxWorkers = 4u;
+	const uint32 NumWorkerThreads = FTaskGraphInterface::Get().GetNumWorkerThreads();
+	const uint32 MaxTasks = FApp::ShouldUseThreadingForPerformance() ? FMath::Min<uint32>(NumWorkerThreads, MaxWorkers) : 1u;
+
+	struct FTaskContext
+	{
+		uint32 WorkerIndex = 0;
+	};
+
+	TArray<FTaskContext, TInlineAllocator<MaxWorkers>> TaskContexts;
+	for (uint32 WorkerIndex = 0; WorkerIndex < MaxTasks; ++WorkerIndex)
+	{
+		TaskContexts.Add(FTaskContext{ WorkerIndex });
+	}
+
+	FD3D12ExplicitDescriptorCache TransientDescriptorCache(GetParentDevice(), MaxTasks /* Worker Count */);
+	TransientDescriptorCache.Init(0, NumViewDescriptors, NumSamplerDescriptors, ERHIBindlessConfiguration::AllShaders);
+
+	TArray<FShaderBundleBinderOps, TInlineAllocator<MaxWorkers>> BinderOps;
+	BinderOps.SetNum(MaxTasks);
+
+	TResourceArray<uint32> LocalRootArgs;
+	uint32 MinRootArgBufferSizeInDWords = (Pipeline->RootArgStrideInBytes / 4) * (Pipeline->MaxRootArgOffset + 1);
+	LocalRootArgs.AddZeroed(MinRootArgBufferSizeInDWords);
+
+	FAllocatedConstantBuffer SharedConstantBuffer(*this);
+	SetShaderBundleSharedBindlessConstants(*this, SharedBindlessParameters, SharedConstantBuffer);
+
+	auto RecordTask = [this, &LocalRootArgs, Pipeline, &TransientDescriptorCache, &ValidRecords, &Dispatches, &BinderOps, &SharedConstantBuffer](FTaskContext& Context, int32 RecordIndex)
+	{
+		uint32 DispatchIndex = ValidRecords[RecordIndex];
+		const FRHIShaderBundleGraphicsDispatch& Dispatch = Dispatches[DispatchIndex];
+		check(Dispatch.IsValid());
+
+		const uint32 ShaderTableIndex = RecordIndex + 1;
+		check(Pipeline->RootArgOffsets.IsValidIndex(ShaderTableIndex));
+		const uint32 RootArgOffset = Pipeline->RootArgOffsets[ShaderTableIndex];
+		check((Pipeline->RootArgStrideInBytes / 4) * (RootArgOffset + 1) <= (uint32)LocalRootArgs.Num());
+
+		FRHIWorkGraphShader* MeshShader = (FRHIWorkGraphShader*)Dispatch.PipelineInitializer.BoundShaderState.GetMeshShader();
+		FD3D12WorkGraphShader* D3D12WorkGraphShader = FD3D12DynamicRHI::ResourceCast(MeshShader);
+		FRHIPixelShader* PixelShader = Dispatch.PipelineInitializer.BoundShaderState.GetPixelShader();
+		FD3D12PixelShader* D3D12PixelShader = FD3D12DynamicRHI::ResourceCast(PixelShader);
+		FD3D12RootSignature const* LocalRootSignature = GetParentAdapter()->GetWorkGraphGraphicsRootSignature(Dispatch.PipelineInitializer.BoundShaderState);
+
+		RecordBindings(
+			*this,
+			SF_Pixel,
+			TransientDescriptorCache,
+			BinderOps[Context.WorkerIndex],
+			Context.WorkerIndex,
+			PixelShader,
+			D3D12PixelShader,
+			*Dispatch.Parameters_PS,
+			LocalRootSignature,
+			SharedConstantBuffer,
+			Dispatch.Constants,
+			MakeArrayView(&LocalRootArgs[RootArgOffset * Pipeline->RootArgStrideInBytes / 4], Pipeline->RootArgStrideInBytes / 4)
+		);
+
+		RecordBindings(
+			*this,
+			SF_Mesh,
+			TransientDescriptorCache,
+			BinderOps[Context.WorkerIndex],
+			Context.WorkerIndex,
+			MeshShader,
+			D3D12WorkGraphShader,
+			*Dispatch.Parameters_MSVS,
+			LocalRootSignature,
+			SharedConstantBuffer,
+			Dispatch.Constants,
+			MakeArrayView(&LocalRootArgs[RootArgOffset * Pipeline->RootArgStrideInBytes / 4], Pipeline->RootArgStrideInBytes / 4)
+		);
+	};
+
+	const int32 ItemsPerTask = 1024;
+	ParallelForWithExistingTaskContext(TEXT("DispatchShaderBundle"), MakeArrayView(TaskContexts), ValidRecords.Num(), ItemsPerTask, RecordTask);
+
+	// Apply Binder Ops
+	{
+		for (int32 WorkerIndex = 1; WorkerIndex < BinderOps.Num(); ++WorkerIndex)
+		{
+			for (FD3D12ShaderResourceView* SRV : BinderOps[WorkerIndex].TransitionSRVs)
+			{
+				BinderOps[0].AddResourceTransition(SRV);
+			}
+
+			for (FD3D12UnorderedAccessView* UAV : BinderOps[WorkerIndex].TransitionUAVs)
+			{
+				BinderOps[0].AddResourceTransition(UAV);
+			}
+
+			BinderOps[WorkerIndex].TransitionSRVs.Empty();
+			BinderOps[WorkerIndex].TransitionUAVs.Empty();
+			BinderOps[WorkerIndex].TransitionViewSet.Empty();
+
+			BinderOps[WorkerIndex].ClearUAVs.Empty();
+			BinderOps[WorkerIndex].TransitionClearSet.Empty();
+		}
+
+		for (FD3D12UnorderedAccessView* UAV : BinderOps[0].ClearUAVs)
+		{
+			ClearShaderResources(UAV, EShaderParameterTypeMask::SRVMask);
+		}
+
+//TODO: check resource view
+// 		for (FD3D12ShaderResourceView* SRV : BinderOps[0].TransitionSRVs)
+// 		{
+// 			TransitionResource(SRV, D3D12_RESOURCE_STATE_NON_PIXEL_SHADER_RESOURCE);
+// 		}
+// 
+// 		for (FD3D12UnorderedAccessView* UAV : BinderOps[0].TransitionUAVs)
+// 		{
+// 			TransitionResource(UAV, D3D12_RESOURCE_STATE_UNORDERED_ACCESS);
+//		}
+	}
+
+	FlushResourceBarriers();
+
+	// Create SRV for RecordArgsBuffer
+	FD3D12Buffer* RecordArgBufferPtr = FD3D12DynamicRHI::ResourceCast(RecordArgBuffer);
+
+	D3D12_SHADER_RESOURCE_VIEW_DESC SRVDesc = {};
+	SRVDesc.ViewDimension = D3D12_SRV_DIMENSION_BUFFER;
+	SRVDesc.Format = DXGI_FORMAT_R32_TYPELESS;
+	SRVDesc.Shader4ComponentMapping = D3D12_DEFAULT_SHADER_4_COMPONENT_MAPPING;
+	SRVDesc.Buffer.FirstElement = 0;
+	SRVDesc.Buffer.NumElements = RecordArgBufferPtr->GetSize() >> 2u;
+	SRVDesc.Buffer.StructureByteStride = 0;
+	SRVDesc.Buffer.Flags = D3D12_BUFFER_SRV_FLAG_RAW;
+
+	TSharedPtr<FD3D12ShaderResourceView> RecordArgBufferSRV;
+	RecordArgBufferSRV = MakeShared<FD3D12ShaderResourceView>(GetParentDevice(), nullptr); // Always single GPU object, so FirstLinkedObject is nullptr
+	RecordArgBufferSRV->CreateView(RecordArgBufferPtr, SRVDesc, FD3D12ShaderResourceView::EFlags::None);
+
+	// Gather root arguments for shader bundle entry node.
+	uint32 RecordArgBufferBindlessHandle = 0;
+	uint32 BindSlot = 0;
+	D3D12_GPU_DESCRIPTOR_HANDLE ResourceDescriptorTableBaseGPU = {};
+
+	if (bBindlessResources)
+	{
+		RecordArgBufferBindlessHandle = RecordArgBufferSRV->GetBindlessHandle().GetIndex();
+		check(RecordArgBufferBindlessHandle != INDEX_NONE);
+
+		StateCache.ApplyBindlessResources(nullptr, SF_Compute, SF_NumStandardFrequencies);
+	}
+	else
+	{
+		D3D12_CPU_DESCRIPTOR_HANDLE LocalSRVs[MAX_SRVS];
+		LocalSRVs[WorkGraphGlobalShader->RecordArgBufferParam.GetBaseIndex()] = RecordArgBufferSRV->GetOfflineCpuHandle();
+
+		const int32 DescriptorTableBaseIndex = TransientDescriptorCache.Allocate(LocalSRVs, 1, D3D12_DESCRIPTOR_HEAP_TYPE_CBV_SRV_UAV, 0);
+		ResourceDescriptorTableBaseGPU = TransientDescriptorCache.ViewHeap.GetDescriptorGPU(DescriptorTableBaseIndex);
+
+		BindSlot = WorkGraphGlobalShaderRHI->RootSignature->SRVRDTBindSlot(SF_Compute);
+		const uint32 BindSlotOffset = WorkGraphGlobalShaderRHI->RootSignature->GetBindSlotOffsetInBytes(BindSlot) / 4;
+
+		TArrayView<uint32> RootArgSlice = MakeArrayView(&LocalRootArgs[Pipeline->RootArgOffsets[0] * Pipeline->RootArgStrideInBytes / 4], Pipeline->RootArgStrideInBytes / 4);
+		FMemory::Memcpy(&RootArgSlice[BindSlotOffset], &ResourceDescriptorTableBaseGPU, sizeof(ResourceDescriptorTableBaseGPU));
+	}
+
+	// Upload local root arguments table.
+	D3D12_GPU_VIRTUAL_ADDRESS_RANGE_AND_STRIDE NodeLocalRootArgumentsTable{ 0, 0, 0 };
+	if (LocalRootArgs.Num())
+	{
+		// todo: Check if copy queue is the optimal way to upload the root args.
+		// todo: Use a single buffer owned by the shader bundle RHI object (needs a copy operation that doesn't complain about multiple uploads).
+		D3D12_RESOURCE_DESC Desc = CD3DX12_RESOURCE_DESC::Buffer(LocalRootArgs.GetResourceDataSize(), D3D12_RESOURCE_FLAG_NONE);
+		FD3D12Buffer* RootArgBuffer = GetParentDevice()->GetParentAdapter()->CreateRHIBuffer(
+			Desc, 16, FRHIBufferDesc(Desc.Width, 0, BUF_Static), ED3D12ResourceStateMode::MultiState,
+			D3D12_RESOURCE_STATE_COPY_DEST, true, FRHIGPUMask::FromIndex(GetParentDevice()->GetGPUIndex()), nullptr, TEXT("BundleRecordBuffer"), ERHIAccess::CopyDest);
+
+		BatchedSyncPoints.ToWait.Emplace(RootArgBuffer->UploadResourceDataViaCopyQueue(&LocalRootArgs));
+		TransitionResource(RootArgBuffer->GetResource(), D3D12_RESOURCE_STATE_COPY_DEST, D3D12_RESOURCE_STATE_COMMON, 0);
+
+		NodeLocalRootArgumentsTable.StartAddress = RootArgBuffer->ResourceLocation.GetGPUVirtualAddress();
+		NodeLocalRootArgumentsTable.SizeInBytes = RootArgBuffer->ResourceLocation.GetSize();
+		NodeLocalRootArgumentsTable.StrideInBytes = Pipeline->RootArgStrideInBytes;
+	}
+
+	// Apply the transient descriptor heaps.
+	// Note this only uses transient heap for non-bindless.
+	SetExplicitDescriptorCache(TransientDescriptorCache);
+
+	// Set graphics state
+	GraphicsCommandList()->SetGraphicsRootSignature(Pipeline->RootSignature);
+
+	if (!bBindlessResources)
+	{
+		GraphicsCommandList()->SetGraphicsRootDescriptorTable(BindSlot, ResourceDescriptorTableBaseGPU);
+	}
+
+	GraphicsCommandList()->IASetVertexBuffers(0, D3D12_IA_VERTEX_INPUT_RESOURCE_SLOT_COUNT, nullptr);
+	GraphicsCommandList()->IASetIndexBuffer(nullptr);
+	GraphicsCommandList()->OMSetRenderTargets(0, nullptr, 0, nullptr);
+
+	D3D12_VIEWPORT Viewport{};
+	{
+		Viewport.TopLeftX = BundleState.ViewRect.Min.X;
+		Viewport.TopLeftY = BundleState.ViewRect.Min.Y;
+		Viewport.Width = BundleState.ViewRect.Width();
+		Viewport.Height = BundleState.ViewRect.Height();
+		Viewport.MinDepth = BundleState.DepthMin;
+		Viewport.MaxDepth = BundleState.DepthMax;
+	}
+	GraphicsCommandList()->RSSetViewports(1, &Viewport);
+
+	D3D12_RECT Rect{};
+	{
+		Rect.left = BundleState.ViewRect.Min.X;
+		Rect.top = BundleState.ViewRect.Min.Y;
+		Rect.right = BundleState.ViewRect.Max.X;
+		Rect.bottom = BundleState.ViewRect.Max.Y;
+	}
+	GraphicsCommandList()->RSSetScissorRects(1, &Rect);
+
+	const D3D_PRIMITIVE_TOPOLOGY PrimitiveTopology = TranslatePrimitiveType(BundleState.PrimitiveType);
+	GraphicsCommandList()->IASetPrimitiveTopology(PrimitiveTopology);
+
+	GraphicsCommandList()->OMSetStencilRef(BundleState.StencilRef);
+	GraphicsCommandList()->OMSetBlendFactor(BundleState.BlendFactor);
+//	GraphicsCommandList()->OMSetDepthBounds(BundleState.DepthMin, BundleState.DepthMax);
+
+	// Kick off the work graph	
+	D3D12_SET_PROGRAM_DESC SetProgramDesc = {};
+	SetProgramDesc.Type = D3D12_PROGRAM_TYPE_WORK_GRAPH;
+	SetProgramDesc.WorkGraph.ProgramIdentifier = Pipeline->ProgramIdentifier;
+	SetProgramDesc.WorkGraph.Flags = D3D12_SET_WORK_GRAPH_FLAG_INITIALIZE;
+	SetProgramDesc.WorkGraph.BackingMemory = Pipeline->BackingMemoryAddressRange;
+	SetProgramDesc.WorkGraph.NodeLocalRootArgumentsTable = NodeLocalRootArgumentsTable;
+	GraphicsCommandList10()->SetProgram(&SetProgramDesc);
+
+	const uint32 RecordCount = Dispatches.Num();
+	FDispatchShaderBundleWorkGraph::FEntryNodeRecord InputRecord = FDispatchShaderBundleWorkGraph::MakeInputRecord(RecordCount, ShaderBundle->ArgOffset, ShaderBundle->ArgStride, RecordArgBufferBindlessHandle);
+
+	if (!GShaderBundleSkipDispatch)
+	{
+		D3D12_DISPATCH_GRAPH_DESC DispatchGraphDesc = {};
+		DispatchGraphDesc.Mode = D3D12_DISPATCH_MODE_NODE_CPU_INPUT;
+		DispatchGraphDesc.NodeCPUInput.EntrypointIndex = 0;
+		DispatchGraphDesc.NodeCPUInput.NumRecords = 1;
+		DispatchGraphDesc.NodeCPUInput.RecordStrideInBytes = sizeof(InputRecord);
+		DispatchGraphDesc.NodeCPUInput.pRecords = &InputRecord;
+		GraphicsCommandList10()->DispatchGraph(&DispatchGraphDesc);
+	}
+
+	// Restore global descriptor heaps if necessary.
+	UnsetExplicitDescriptorCache();
+
+	// We did not write through the state cache, so we need to invalidate it so subsequent workloads correctly re-bind state
+	StateCache.DirtyState();
+
+	ConditionalSplitCommandList();
+
+#endif // D3D12_RHI_WORKGRAPHS
 }
diff --git a/Engine/Source/Runtime/RHI/Public/RHIResources.h b/Engine/Source/Runtime/RHI/Public/RHIResources.h
--- a/Engine/Source/Runtime/RHI/Public/RHIResources.h
+++ b/Engine/Source/Runtime/RHI/Public/RHIResources.h
@@ -4508,24 +4508,28 @@ public:
 	{
 		return NameHash == Rhs.NameHash && 
 			NameTableHash == Rhs.NameTableHash && 
-			ShaderTableHash == Rhs.ShaderTableHash;
+			ShaderTableHash == Rhs.ShaderTableHash &&
+			GraphicsPSOTableHash == Rhs.GraphicsPSOTableHash;
 	}
 
 	friend uint32 GetTypeHash(const FWorkGraphPipelineStateSignature& Initializer)
 	{
 		return GetTypeHash(Initializer.NameHash) ^ 
 			GetTypeHash(Initializer.NameTableHash) ^ 
-			GetTypeHash(Initializer.ShaderTableHash);
+			GetTypeHash(Initializer.ShaderTableHash) ^
+			GetTypeHash(Initializer.GraphicsPSOTableHash);
 	}
 
 	uint64 GetNameHash() const { return NameHash; }
 	uint64 GetNameTableHash() const { return NameTableHash; }
 	uint64 GetShaderTableHash() const { return ShaderTableHash; }
+	uint64 GetGraphicsPSOTableHash() const { return GraphicsPSOTableHash; }
 
 protected:
 	uint64 NameHash = 0;
 	uint64 NameTableHash = 0;
 	uint64 ShaderTableHash = 0;
+	uint64 GraphicsPSOTableHash = 0;
 };
 
 class FWorkGraphPipelineStateInitializer : public FWorkGraphPipelineStateSignature
@@ -4573,16 +4577,37 @@ public:
 		RootShaderIndex = InRootShaderIndex;
 	}
 
+	void SetGraphicsPSOTable(const TArrayView<FGraphicsPipelineStateInitializer const*>& InGraphicsPSOs, uint64 Hash = 0)
+	{
+		GraphicsPSOTable = InGraphicsPSOs;
+		GraphicsPSOTableHash = Hash ? Hash : ComputeGraphicsPSOTableHash(InGraphicsPSOs);
+	}
+
 	FString const& GetProgramName() const { return ProgramName; }
 	TArray<FNameMap> const& GetNameTable() const { return NameMaps; }
 	int32 GetRootShaderIndex() const { return RootShaderIndex; }
 	TArrayView<FRHIWorkGraphShader*> const& GetShaderTable() const { return ShaderTable; }
+	TArrayView<FGraphicsPipelineStateInitializer const*> const& GetGraphicsPSOTable() const { return GraphicsPSOTable; }
 
 private:
+	uint64 ComputeGraphicsPSOTableHash(const TArrayView<FGraphicsPipelineStateInitializer const*>& InGraphicsPSOTable, uint64 InitialHash = 5699878132332235837ull)
+	{
+		uint64 CombinedHash = InitialHash;
+		for (FGraphicsPipelineStateInitializer const* GraphicsPSO : InGraphicsPSOTable)
+		{
+			uint64 GraphicsPSOHash = GraphicsPSO != nullptr ? GraphicsPSO->StatePrecachePSOHash : 0;
+			// 64 bit hash combination as per boost::hash_combine_impl
+			CombinedHash ^= GraphicsPSOHash + 0x9e3779b97f4a7c15ull + (CombinedHash << 12) + (CombinedHash >> 4);
+		}
+
+		return CombinedHash;
+	}
+
 	FString ProgramName;
 	TArray<FNameMap> NameMaps;
 	int32 RootShaderIndex = -1;
 	TArrayView<FRHIWorkGraphShader*> ShaderTable;
+	TArrayView<FGraphicsPipelineStateInitializer const*> GraphicsPSOTable;
 };
 
 class FRayTracingPipelineStateSignature
