Temporarily disabling Async Compute on Windows while investigating UE-233616 (Screen Corruptions) that appears to be related to the DX12 barrier refactor.

diff --git a/Engine/Source/Runtime/D3D12RHI/Private/D3D12Adapter.cpp b/Engine/Source/Runtime/D3D12RHI/Private/D3D12Adapter.cpp
--- a/Engine/Source/Runtime/D3D12RHI/Private/D3D12Adapter.cpp
+++ b/Engine/Source/Runtime/D3D12RHI/Private/D3D12Adapter.cpp
@@ -1158,6 +1158,10 @@ void FD3D12Adapter::InitializeDevices()
 #endif
 		}
 
+#if PLATFORM_WINDOWS
+		GSupportsEfficientAsyncCompute = false;
+#endif
+
 		if (!FPlatformMemory::SupportsFastVRAMMemory() && GSupportsEfficientAsyncCompute)
 		{
 			GRHIGlobals.SupportsAsyncComputeTransientAliasing = true;
