Optimize used d3d12 resource tracking for persistent SBTs for residency tracking:
- track used d3d12 resources per record instead of d3d12 base shader resources (base shader resources still had to be iterated at commit to extract the actual d3d12 resources which is expensive operation because there can easily be 1000s of these)
- persistent SBT records now need listener support for base shader resource renames (defrag/streamin/out) and rhi texture reference replace (stream in/out) to update the used records

diff --git a/Engine/Source/Runtime/D3D12RHI/Private/D3D12RayTracing.cpp b/Engine/Source/Runtime/D3D12RHI/Private/D3D12RayTracing.cpp
--- a/Engine/Source/Runtime/D3D12RHI/Private/D3D12RayTracing.cpp
+++ b/Engine/Source/Runtime/D3D12RHI/Private/D3D12RayTracing.cpp
@@ -1585,55 +1585,30 @@ public:
 				for (FD3D12Resource* Resource : WorkerData[WorkerIndex].DynamicReferencedResources)
 				{
 					AddDynamicReferencedResource(0, Resource);
-				}
-				for (auto Iter = WorkerData[WorkerIndex].PersistentReferencedShaderResources.CreateIterator(); Iter; ++Iter)
-				{
-					UpdatePersistentReferencedShaderResource(0, Iter.Key(), Iter.Value());
-				}
-				for (auto Iter = WorkerData[WorkerIndex].PersistentReferencedTextureReferences.CreateIterator(); Iter; ++Iter)
+				}				
+				for (auto Iter = WorkerData[WorkerIndex].PersistentReferencedResources.CreateIterator(); Iter; ++Iter)
 				{
-					UpdatePersistentReferencedTextureReference(0, Iter.Key(), Iter.Value());
+					UpdatePersistentReferencedResource(0, Iter.Key(), Iter.Value());
 				}
 				WorkerData[WorkerIndex].DynamicReferencedResources.Empty();
 				WorkerData[WorkerIndex].DynamicReferencedResourceSet.Empty();
-				WorkerData[WorkerIndex].PersistentReferencedShaderResources.Empty();
-				WorkerData[WorkerIndex].PersistentReferencedTextureReferences.Empty();
+				WorkerData[WorkerIndex].PersistentReferencedResources.Empty();
 			}
 		}
 
 		// Build the used persistent d3d12 resource set from the views
 		{
 			TRACE_CPUPROFILER_EVENT_SCOPE(BuildUsedResourceSet);
+			
+			ReferencedResources.Empty(WorkerData[0].PersistentReferencedResources.Num() + WorkerData[0].DynamicReferencedResources.Num());
 
-			// Pretty expensive for a lot of resources - could be kicked in async task or tracking code refactored to
-			// directly track the used d3d12 resources but then SBT needs to be listener to all used base shader resources
-			// for defrag and texture reference changes
-			TSet<FD3D12Resource*> ReferencedResourcesSet;
-			ReferencedResourcesSet.Reserve(ReferencedResources.Num());
-			for (auto Iter = WorkerData[0].PersistentReferencedShaderResources.CreateIterator(); Iter; ++Iter)
-			{
-				FD3D12BaseShaderResource* ShaderResource = Iter.Key();
-				// buffer could have been streamed out but MDC is still valid containing the RHI buffer due to internal swap
-				if (ShaderResource->GetResource())
-				{
-					ReferencedResourcesSet.Add(ShaderResource->GetResource());
-				}
-			}
-			for (auto Iter = WorkerData[0].PersistentReferencedTextureReferences.CreateIterator(); Iter; ++Iter)
-			{
-				FD3D12RHITextureReference* TextureReference = Iter.Key();
-				FD3D12Texture* Texture = FD3D12CommandContext::RetrieveTexture(TextureReference, Device->GetGPUIndex());
-				check(Texture->GetResource());
-				ReferencedResourcesSet.Add(Texture->GetResource());
-			}
-
-			ReferencedResources.Empty(ReferencedResourcesSet.Num() + WorkerData[0].DynamicReferencedResources.Num());
-			for (TSet<FD3D12Resource*>::TConstIterator SetIt(ReferencedResourcesSet); SetIt; ++SetIt)
+			// Add tracked resources for persistent SBT records
+			for (auto Iter = WorkerData[0].PersistentReferencedResources.CreateIterator(); Iter; ++Iter)
 			{
-				ReferencedResources.Add(*SetIt);
+				ReferencedResources.Add(Iter.Key());
 			}
 
-			// Also add all the collected dynamic resources and reset for the next frame			
+			// Also add all the collected dynamic resources and reset for the next 'frame'			
 			ReferencedResources.Append(WorkerData[0].DynamicReferencedResources);
 			WorkerData[0].DynamicReferencedResources.Empty();
 			WorkerData[0].DynamicReferencedResourceSet.Empty();
@@ -1870,6 +1845,144 @@ public:
 		}
 	};
 
+#if ENABLE_RESIDENCY_MANAGEMENT
+	struct FRecordUpdateShaderResourceRenameListener : public FD3D12ShaderResourceRenameListener
+	{
+		FD3D12RayTracingShaderBindingTableInternal& ShaderTable;
+		uint32 RecordIndex = 0;
+		TRefCountPtr<FD3D12BaseShaderResource> ShaderResource;
+		FD3D12Resource* CurrentResource = nullptr;
+
+		FRecordUpdateShaderResourceRenameListener(FD3D12RayTracingShaderBindingTableInternal& InShaderTable) :	ShaderTable(InShaderTable)
+		{
+		}
+		virtual ~FRecordUpdateShaderResourceRenameListener() = default;
+
+		void Init(uint32 WorkerIndex, uint32 InRecordIndex, FD3D12BaseShaderResource* InShaderResource)
+		{
+			RecordIndex = InRecordIndex;
+			ShaderResource = InShaderResource;
+			CurrentResource = InShaderResource->ResourceLocation.GetResource();
+
+			ShaderResource->AddRenameListener(this);
+			if (CurrentResource)
+			{
+				ShaderTable.UpdatePersistentReferencedResource(WorkerIndex, CurrentResource, 1);
+			}
+		}
+
+		void Reset(uint32 WorkerIndex)
+		{
+			check(ShaderResource);
+			if (CurrentResource)
+			{
+				ShaderTable.UpdatePersistentReferencedResource(WorkerIndex, CurrentResource, -1);
+			}
+			ShaderResource->RemoveRenameListener(this);
+			ShaderResource = nullptr;
+			RecordIndex = 0;
+			CurrentResource = 0;
+		}
+
+		virtual void ResourceRenamed(FD3D12ContextArray const& Contexts, FD3D12BaseShaderResource* InRenamedResource, FD3D12ResourceLocation* InNewResourceLocation) override
+		{
+			check(InRenamedResource == ShaderResource);
+			FD3D12Resource* NewResource = InNewResourceLocation->GetResource();
+			
+			// only care if actual d3d12 resource has changed
+			if (NewResource != CurrentResource)
+			{
+				// Resource could be streamed in or out and then contain empty d3d12 backing resource
+				if (CurrentResource)
+				{
+					ShaderTable.UpdatePersistentReferencedResource(0, CurrentResource, -1);
+				}
+				if (NewResource)
+				{
+					ShaderTable.UpdatePersistentReferencedResource(0, NewResource, 1);
+				}
+				CurrentResource = NewResource;
+			}
+		}
+	};
+	
+	struct FRecordUpdateTextureReferenceReplaceListener : public FD3D12TextureReferenceReplaceListener, public FD3D12ShaderResourceRenameListener
+	{
+		FD3D12RayTracingShaderBindingTableInternal& ShaderTable;
+		uint32 RecordIndex = 0;
+		TRefCountPtr<FD3D12RHITextureReference> TextureReference;
+		FD3D12Resource* CurrentResource = nullptr;
+
+		FRecordUpdateTextureReferenceReplaceListener(FD3D12RayTracingShaderBindingTableInternal& InShaderTable) :	ShaderTable(InShaderTable)
+		{
+		}
+		virtual ~FRecordUpdateTextureReferenceReplaceListener() = default;
+
+		void Init(uint32 WorkerIndex, uint32 InRecordIndex, FD3D12RHITextureReference* InTextureReference)
+		{
+			RecordIndex = InRecordIndex;
+			TextureReference = InTextureReference;
+			TextureReference->AddReplaceListener(this);			
+			
+			// Handle rename events on the references texture as well
+			FD3D12Texture* ReferencedTexture = FD3D12DynamicRHI::ResourceCast(TextureReference->GetReferencedTexture());
+			ReferencedTexture->AddRenameListener(this);
+
+			CurrentResource = ReferencedTexture->ResourceLocation.GetResource();
+			check(CurrentResource);
+
+			ShaderTable.UpdatePersistentReferencedResource(WorkerIndex, CurrentResource, 1);
+		}
+
+		void Reset(uint32 WorkerIndex)
+		{
+			check(TextureReference && CurrentResource);
+			ShaderTable.UpdatePersistentReferencedResource(WorkerIndex, CurrentResource, -1);
+			
+			FD3D12Texture* ReferencedTexture = FD3D12DynamicRHI::ResourceCast(TextureReference->GetReferencedTexture());
+			ReferencedTexture->RemoveRenameListener(this);
+			TextureReference->RemoveReplaceListener(this);
+
+			TextureReference = nullptr;
+			RecordIndex = 0;
+			CurrentResource = 0;
+		}
+
+		virtual void TextureReplaced(FD3D12ContextArray const& Contexts, FD3D12RHITextureReference* InTextureReference, FD3D12Texture* CurrentTexture, FD3D12Texture* NewTexture) override
+		{
+			check(TextureReference == InTextureReference);
+			FD3D12Resource* NewResource = NewTexture->ResourceLocation.GetResource();
+			
+			// only care if actual d3d12 resource has changed
+			if (NewResource != CurrentResource)
+			{
+				ShaderTable.UpdatePersistentReferencedResource(0, CurrentResource, -1);
+				ShaderTable.UpdatePersistentReferencedResource(0, NewResource, 1);
+				CurrentResource = NewResource;
+			}
+			
+			// Also update the rename listener
+			CurrentTexture->RemoveRenameListener(this);
+			NewTexture->AddRenameListener(this);
+		}
+
+		virtual void ResourceRenamed(FD3D12ContextArray const& Contexts, FD3D12BaseShaderResource* InRenamedResource, FD3D12ResourceLocation* InNewResourceLocation) override
+		{
+			FD3D12Texture* ReferencedTexture = FD3D12DynamicRHI::ResourceCast(TextureReference->GetReferencedTexture());
+			check(InRenamedResource == ReferencedTexture);
+			FD3D12Resource* NewResource = InNewResourceLocation->GetResource();
+			
+			// only care if actual d3d12 resource has changed
+			if (NewResource != CurrentResource)
+			{
+				ShaderTable.UpdatePersistentReferencedResource(0, CurrentResource, -1);
+				ShaderTable.UpdatePersistentReferencedResource(0, NewResource, 1);
+				CurrentResource = NewResource;
+			}
+		}
+	};
+#endif // ENABLE_RESIDENCY_MANAGEMENT
+
 	// Tracked data for each record so stored GPU data can be updated after defrag operation
 	struct FRecordData
 	{
@@ -1883,8 +1996,8 @@ public:
 		void Reset()
 		{
 #if ENABLE_RESIDENCY_MANAGEMENT
-			ReferencedShaderResources.Empty();
-			ReferencedTextureReferences.Empty();
+			ReferencedShaderResourceListeners.Empty();
+			ReferencedTextureReferenceListeners.Empty();
 #endif // ENABLE_RESIDENCY_MANAGEMENT
 			UpdateGeometryListener.Reset();
 			UpdateUniformBufferListeners.Empty();
@@ -1900,9 +2013,9 @@ public:
 		FRecordUpdateRayTracingGeometryListener UpdateGeometryListener;
 		FD3D12ResourceLocation* LooseParameterResourceLocation = nullptr;
 
-#if ENABLE_RESIDENCY_MANAGEMENT		
-		TArray<FD3D12BaseShaderResource*> ReferencedShaderResources;
-		TArray<FD3D12RHITextureReference*> ReferencedTextureReferences;
+#if ENABLE_RESIDENCY_MANAGEMENT				
+		TArray<FRecordUpdateShaderResourceRenameListener*> ReferencedShaderResourceListeners;
+		TArray<FRecordUpdateTextureReferenceReplaceListener*> ReferencedTextureReferenceListeners;
 #endif // ENABLE_RESIDENCY_MANAGEMENT		
 #if DO_CHECK
 		TArray<uint8> LooseParameterData;
@@ -1956,26 +2069,12 @@ public:
 	}
 		
 #if ENABLE_RESIDENCY_MANAGEMENT	
-	void UpdatePersistentReferencedShaderResource(uint32 WorkerIndex, FD3D12BaseShaderResource* ShaderResource, int32 Count)
+	void UpdatePersistentReferencedResource(uint32 WorkerIndex, FD3D12Resource* Resource, int32 Count)
 	{
-		// Update the worker thread reference count of the resource
-		int32& ReferenceCount = WorkerData[WorkerIndex].PersistentReferencedShaderResources.FindOrAdd(ShaderResource);
-		ReferenceCount += Count;
-		
-		if (WorkerIndex == 0)
-		{
-			check(ReferenceCount >= 0);
-			if (ReferenceCount == 0)
-			{
-				WorkerData[WorkerIndex].PersistentReferencedShaderResources.Remove(ShaderResource);
-			}
-		}
-	}
+		check(Resource);
 
-	void UpdatePersistentReferencedTextureReference(uint32 WorkerIndex, FD3D12RHITextureReference* TextureReference, int32 Count)
-	{
 		// Update the worker thread reference count of the resource
-		int32& ReferenceCount = WorkerData[WorkerIndex].PersistentReferencedTextureReferences.FindOrAdd(TextureReference);
+		int32& ReferenceCount = WorkerData[WorkerIndex].PersistentReferencedResources.FindOrAdd(Resource);
 		ReferenceCount += Count;
 		
 		if (WorkerIndex == 0)
@@ -1983,7 +2082,7 @@ public:
 			check(ReferenceCount >= 0);
 			if (ReferenceCount == 0)
 			{
-				WorkerData[WorkerIndex].PersistentReferencedTextureReferences.Remove(TextureReference);
+				WorkerData[WorkerIndex].PersistentReferencedResources.Remove(Resource);
 			}
 		}
 	}
@@ -2008,10 +2107,20 @@ public:
 		}
 		else if (BindingType == ERayTracingLocalShaderBindingType::Persistent)
 		{
-			// Update the tracked record data so it can be updated in case of defrag operation
+			// Same base shader resource can be bound multiple times to the same record - only track once
 			FRecordData& RecordData = HitRecordData[RecordIndex];
-			RecordData.ReferencedShaderResources.Add(ShaderResource);
-			UpdatePersistentReferencedShaderResource(WorkerIndex, ShaderResource, 1);
+			for (FRecordUpdateShaderResourceRenameListener* Listener : RecordData.ReferencedShaderResourceListeners)
+			{
+				if (Listener->ShaderResource == ShaderResource)
+				{
+					return;
+				}
+			}
+
+			// Update the tracked record data so it can be updated in case of defrag operation
+			FRecordUpdateShaderResourceRenameListener* Listener = new FRecordUpdateShaderResourceRenameListener(*this);
+			Listener->Init(WorkerIndex, RecordIndex, ShaderResource);
+			RecordData.ReferencedShaderResourceListeners.Add(Listener);
 		}
 #endif // ENABLE_RESIDENCY_MANAGEMENT
 	}
@@ -2025,22 +2134,32 @@ public:
 			AddDynamicReferencedResource(WorkerIndex, Texture->ResourceLocation.GetResource());
 		}
 		else if (BindingType == ERayTracingLocalShaderBindingType::Persistent)
-		{
-			// Update the tracked record data so it can be updated in case of texture reference swap
-			FRecordData& RecordData = HitRecordData[RecordIndex];
-			
+		{			
 			// Collect all texture references independent of the shader resources because the internal shader resource can be changed on swap
 			if (FRHITextureReference* TextureReferenceRHI = RHITexture->GetTextureReference())
 			{
 				FD3D12RHITextureReference* TextureReference = FD3D12CommandContext::RetrieveObject<FD3D12RHITextureReference>(TextureReferenceRHI, Device->GetGPUIndex());
-				RecordData.ReferencedTextureReferences.Add(TextureReference);
-				UpdatePersistentReferencedTextureReference(WorkerIndex, TextureReference, 1);
+								
+				// Same texture reference can be bound multiple times to the same record - only track once
+				FRecordData& RecordData = HitRecordData[RecordIndex];
+				for (FRecordUpdateTextureReferenceReplaceListener* Listener : RecordData.ReferencedTextureReferenceListeners)
+				{
+					if (Listener->TextureReference == TextureReference)
+					{
+						return;
+					}
+				}
+
+				// Update the tracked record data so it can be updated in case of texture reference swap
+				FRecordUpdateTextureReferenceReplaceListener* Listener = new FRecordUpdateTextureReferenceReplaceListener(*this);
+				Listener->Init(WorkerIndex, RecordIndex, TextureReference);
+				RecordData.ReferencedTextureReferenceListeners.Add(Listener);
 			}
 			else
 			{
+				// Fall back to regular base shader resource code path
 				FD3D12Texture* Texture = FD3D12CommandContext::RetrieveTexture(RHITexture, Device->GetGPUIndex());
-				RecordData.ReferencedShaderResources.Add(Texture);
-				UpdatePersistentReferencedShaderResource(WorkerIndex, Texture, 1);
+				AddReferencedShaderResource(WorkerIndex, RecordIndex, BindingType, Texture);
 			}
 		}
 #endif // ENABLE_RESIDENCY_MANAGEMENT
@@ -2156,14 +2275,16 @@ public:
 		}
 		
 #if ENABLE_RESIDENCY_MANAGEMENT
-		for (FD3D12BaseShaderResource* ShaderResource : RecordData.ReferencedShaderResources)
+		for (FRecordUpdateShaderResourceRenameListener* Listener : RecordData.ReferencedShaderResourceListeners)
 		{
-			UpdatePersistentReferencedShaderResource(WorkerIndex, ShaderResource, -1);
+			Listener->Reset(WorkerIndex);
+			delete Listener;
 		}
 
-		for (FD3D12RHITextureReference* TextureReference : RecordData.ReferencedTextureReferences)
+		for (FRecordUpdateTextureReferenceReplaceListener* Listener : RecordData.ReferencedTextureReferenceListeners)
 		{
-			UpdatePersistentReferencedTextureReference(WorkerIndex, TextureReference, -1);
+			Listener->Reset(WorkerIndex);
+			delete Listener;
 		}
 #endif // ENABLE_RESIDENCY_MANAGEMENT
 
@@ -2290,8 +2411,7 @@ public:
 
 #if ENABLE_RESIDENCY_MANAGEMENT
 		// Track all persistent resources with number of refs accross all records (when ref reaches 0 then it doesn't need be tracked anymore)
-		TMap<TRefCountPtr<FD3D12BaseShaderResource>, int32> PersistentReferencedShaderResources;
-		TMap<TRefCountPtr<FD3D12RHITextureReference>, int32> PersistentReferencedTextureReferences;
+		TMap<TRefCountPtr<FD3D12Resource>, int32> PersistentReferencedResources;
 
 		// A set of all resources referenced by this shader table for the purpose of updating residency before ray tracing work dispatch.
 		Experimental::TSherwoodSet<void*> DynamicReferencedResourceSet;
