Serialize EntryPoint in D3D12 PixelShaders.
This is handled using shader optional data and only generated by the DXC compiler.
It is only needed for compiling state objects for work graph graphic nodes. So it is guarded with a D3D12_RHI_WORKGRAPHS_GRAPHICS define.

diff --git a/Engine/Source/Developer/Windows/ShaderFormatD3D/Private/D3DShaderCompilerDXC.cpp b/Engine/Source/Developer/Windows/ShaderFormatD3D/Private/D3DShaderCompilerDXC.cpp
--- a/Engine/Source/Developer/Windows/ShaderFormatD3D/Private/D3DShaderCompilerDXC.cpp
+++ b/Engine/Source/Developer/Windows/ShaderFormatD3D/Private/D3DShaderCompilerDXC.cpp
@@ -1228,12 +1228,7 @@ bool CompileAndProcessD3DShaderDXC(
 			}
 			auto PostSRTWriterCallback = [&](FMemoryWriter& Ar)
 			{
-				if (bIsWorkGraphShader)
-				{
-					FString EntryPoint = Input.EntryPointName;
-					Ar << EntryPoint;
-				}
-				else if (bIsRayTracingShader)
+				if (bIsRayTracingShader)
 				{
 					Ar << RayEntryPoint;
 					Ar << RayAnyHitEntryPoint;
@@ -1295,6 +1290,16 @@ bool CompileAndProcessD3DShaderDXC(
 					uint8 IsSM6 = 1;
 					ShaderCode.AddOptionalData(EShaderOptionalDataKey::ShaderModel6, &IsSM6, 1);
 				}
+
+				// Store EntryPointName for possible use in work graph state object creation.
+				if (bIsWorkGraphShader || Input.Target.GetFrequency() == SF_Pixel)
+				{
+					TArray<uint8> NameData;
+					FMemoryWriter NameWriter(NameData);
+					FString Name = Input.EntryPointName;
+					NameWriter << Name;
+					ShaderCode.AddOptionalData(EShaderOptionalDataKey::EntryPoint, NameData.GetData(), NameData.Num());
+				}
 			};
 
 			// Return a fraction of the number of instructions as DXIL is more verbose than DXBC.
diff --git a/Engine/Source/Developer/Windows/ShaderFormatD3D/Private/ShaderFormatD3D.cpp b/Engine/Source/Developer/Windows/ShaderFormatD3D/Private/ShaderFormatD3D.cpp
--- a/Engine/Source/Developer/Windows/ShaderFormatD3D/Private/ShaderFormatD3D.cpp
+++ b/Engine/Source/Developer/Windows/ShaderFormatD3D/Private/ShaderFormatD3D.cpp
@@ -17,7 +17,7 @@ static FName NAME_PCD3D_SM6(TEXT("PCD3D_SM6"));
 static FName NAME_PCD3D_SM5(TEXT("PCD3D_SM5"));
 static FName NAME_PCD3D_ES3_1(TEXT("PCD3D_ES31"));
 
-static const FGuid UE_SHADER_PCD3D_SHARED_VER = FGuid("6A25E42A-3266-4692-81C4-35B5690B7578");
+static const FGuid UE_SHADER_PCD3D_SHARED_VER = FGuid("B0250A32-89C0-4F5D-9E5A-6DE6E2041829");
 static const FGuid UE_SHADER_PCD3D_SM6_VER    = FGuid("FC317769-CFF7-4007-A045-C1D4CD315BD4");
 static const FGuid UE_SHADER_PCD3D_SM5_VER    = FGuid("5B377D13-C70F-40C5-80C5-C9B228783469");
 static const FGuid UE_SHADER_PCD3D_ES3_1_VER  = FGuid("952939D9-1156-4347-97E9-9FFEA1A9FE14");
diff --git a/Engine/Source/Runtime/D3D12RHI/Private/D3D12Shader.h b/Engine/Source/Runtime/D3D12RHI/Private/D3D12Shader.h
--- a/Engine/Source/Runtime/D3D12RHI/Private/D3D12Shader.h
+++ b/Engine/Source/Runtime/D3D12RHI/Private/D3D12Shader.h
@@ -148,6 +148,10 @@ class FD3D12PixelShader : public FRHIPixelShader, public FD3D12ShaderData
 {
 public:
 	enum { StaticFrequency = SF_Pixel };
+
+#if D3D12_RHI_WORKGRAPHS_GRAPHICS
+	FString EntryPoint;
+#endif
 };
 
 class FD3D12ComputeShader : public FRHIComputeShader, public FD3D12ShaderData
diff --git a/Engine/Source/Runtime/D3D12RHI/Private/D3D12Shaders.cpp b/Engine/Source/Runtime/D3D12RHI/Private/D3D12Shaders.cpp
--- a/Engine/Source/Runtime/D3D12RHI/Private/D3D12Shaders.cpp
+++ b/Engine/Source/Runtime/D3D12RHI/Private/D3D12Shaders.cpp
@@ -7,6 +7,23 @@
 #include "D3D12RHIPrivate.h"
 #include "RHICoreShader.h"
 
+// We don't store EntryPoint name in all shader types to avoid memory bloat.
+// Use template specialization to define which shader types need this.
+template <typename TShaderType>
+constexpr bool CanReadEntryPoint(const TShaderType& InShader) { return false; }
+template <typename TShaderType>
+static void ReadEntryPoint(FBufferReader& InReader, TShaderType& OutShader) {}
+
+// Work graph shaders need to store EntryPoint for state object linking.
+constexpr bool CanReadEntryPoint(const FD3D12WorkGraphShader& InShader) { return true; }
+static void ReadEntryPoint(FBufferReader& InReader, FD3D12WorkGraphShader& OutShader) { InReader << OutShader.EntryPoint; }
+
+// Pixels shaders need to store EntryPoint whenever work graphs are supported because they can be used in work graph graphics nodes.
+#if D3D12_RHI_WORKGRAPHS_GRAPHICS
+constexpr bool CanReadEntryPoint(const FD3D12PixelShader& InShader) { return true; }
+static void ReadEntryPoint(FBufferReader& InReader, FD3D12PixelShader& OutShader) { InReader << OutShader.EntryPoint; }
+#endif
+
 template <typename TShaderType>
 static inline bool ReadShaderOptionalData(FShaderCodeReader& InShaderCode, TShaderType& OutShader)
 {
@@ -61,6 +78,17 @@ static inline bool ReadShaderOptionalData(FShaderCodeReader& InShaderCode, TShad
 		Ar << OutShader.ShaderBindingLayoutHash;
 	}
 
+	if (CanReadEntryPoint(OutShader))
+	{
+		int32 NameSize = 0;
+		uint8 const* NameData = InShaderCode.FindOptionalDataAndSize(EShaderOptionalDataKey::EntryPoint, NameSize);
+		if (NameData && NameSize > 0)
+		{
+			FBufferReader Reader((void*)NameData, NameSize, false);
+			ReadEntryPoint(Reader, OutShader);
+		}
+	}
+
 	UE::RHICore::SetupShaderCodeValidationData(&OutShader, InShaderCode);
 	UE::RHICore::SetupShaderDiagnosticData(&OutShader, InShaderCode);
 
@@ -200,13 +228,7 @@ FComputeShaderRHIRef FD3D12DynamicRHI::RHICreateComputeShader(TArrayView<const u
 
 FWorkGraphShaderRHIRef FD3D12DynamicRHI::RHICreateWorkGraphShader(TArrayView<const uint8> Code, const FSHAHash& Hash, EShaderFrequency ShaderFrequency)
 {
-	auto CustomSerialization = [](FMemoryReaderView& Ar, FD3D12WorkGraphShader* Shader, int32& Offset)
-	{
-		Ar << Shader->EntryPoint;
-		Offset = Ar.Tell();
-	};
-
-	FD3D12WorkGraphShader* Shader = InitStandardShaderWithCustomSerialization(new FD3D12WorkGraphShader(ShaderFrequency), Code, CustomSerialization);
+	FD3D12WorkGraphShader* Shader = InitStandardShader(new FD3D12WorkGraphShader(ShaderFrequency), Code);
 	if (Shader)
 	{
 		Shader->RootSignature = GetAdapter().GetRootSignature(Shader);
diff --git a/Engine/Source/Runtime/D3D12RHI/Private/D3D12WorkGraph.cpp b/Engine/Source/Runtime/D3D12RHI/Private/D3D12WorkGraph.cpp
--- a/Engine/Source/Runtime/D3D12RHI/Private/D3D12WorkGraph.cpp
+++ b/Engine/Source/Runtime/D3D12RHI/Private/D3D12WorkGraph.cpp
@@ -173,7 +173,7 @@ FD3D12WorkGraphPipelineState::FD3D12WorkGraphPipelineState(FD3D12Device* Device,
 			CD3DX12_DXIL_LIBRARY_SUBOBJECT* LibrarySubobject = StateObjectDesc.CreateSubobject<CD3DX12_DXIL_LIBRARY_SUBOBJECT>();
 			CD3DX12_SHADER_BYTECODE ByteCode(D3D12Shader->Code.GetData(), D3D12Shader->Code.Num());
 			LibrarySubobject->SetDXILLibrary(&ByteCode);
-			LibrarySubobject->DefineExport(*PixelShaderName, TEXT("HWRasterizePS")); // todo: Store EntryPoint in pixel shader and use that. Using hard coded value for only use case here!
+			LibrarySubobject->DefineExport(*PixelShaderName, *D3D12Shader->EntryPoint);
 
 			CD3DX12_SUBOBJECT_TO_EXPORTS_ASSOCIATION_SUBOBJECT* ExportAssociationSubobject = StateObjectDesc.CreateSubobject<CD3DX12_SUBOBJECT_TO_EXPORTS_ASSOCIATION_SUBOBJECT>();
 			ExportAssociationSubobject->SetSubobjectToAssociate(*LocalRootSignatureSubobject);
diff --git a/Engine/Source/Runtime/RenderCore/Public/ShaderCore.h b/Engine/Source/Runtime/RenderCore/Public/ShaderCore.h
--- a/Engine/Source/Runtime/RenderCore/Public/ShaderCore.h
+++ b/Engine/Source/Runtime/RenderCore/Public/ShaderCore.h
@@ -749,6 +749,7 @@ enum class EShaderOptionalDataKey : uint8
 	Validation           = uint8('V'),
 	VendorExtension      = uint8('v'),
 	ShaderBindingLayout  = uint8('s'),
+	EntryPoint           = uint8('e'),
 };
 
 enum class EShaderResourceUsageFlags : uint8
