Execute One Command List at a Time on Async Queues with NVIDIA Hardware

The NVIDIA driver drops barriers at the beginning of command lists on the async queue because it assumes each list is executed within its own ECL scope and therefore the GPU is idle'd between each. That's not the case if we execute more than one list at a time, so break it up into individual executions to make it true.

diff --git a/Engine/Source/Runtime/D3D12RHI/Private/D3D12Submission.cpp b/Engine/Source/Runtime/D3D12RHI/Private/D3D12Submission.cpp
--- a/Engine/Source/Runtime/D3D12RHI/Private/D3D12Submission.cpp
+++ b/Engine/Source/Runtime/D3D12RHI/Private/D3D12Submission.cpp
@@ -53,6 +53,30 @@ static FAutoConsoleVariableRef CVarD3D12SubmissionTimeout(
 	TEXT("The maximum time, in seconds, that a submitted GPU command list is allowed to take before the RHI reports a GPU hang"),
 	ECVF_RenderThreadSafe);
 
+static int32 GD3D12SubmissionMaxExecuteBatchSizeDirect = std::numeric_limits<int32>::max();
+static FAutoConsoleVariableRef CVarD3D12SubmissionMaxExecuteBatchSizeDirect(
+	TEXT("r.D3D12.Submission.MaxExecuteBatchSize.Direct"),
+	GD3D12SubmissionMaxExecuteBatchSizeDirect,
+	TEXT("The maximum number of command lists to pass to a single ExecuteCommandLists invocation for direct queues\n")
+	TEXT("The valid range is 1 to INT_MAX inclusive. Values less than 1 will be clamped to 1"),
+	ECVF_RenderThreadSafe);
+
+static int32 GD3D12SubmissionMaxExecuteBatchSizeCopy = std::numeric_limits<int32>::max();
+static FAutoConsoleVariableRef CVarD3D12SubmissionMaxExecuteBatchSizeCopy(
+	TEXT("r.D3D12.Submission.MaxExecuteBatchSize.Copy"),
+	GD3D12SubmissionMaxExecuteBatchSizeCopy,
+	TEXT("The maximum number of command lists to pass to a single ExecuteCommandLists invocation for copy queues\n")
+	TEXT("The valid range is 1 to INT_MAX inclusive. Values less than 1 will be clamped to 1"),
+	ECVF_RenderThreadSafe);
+
+static int32 GD3D12SubmissionMaxExecuteBatchSizeAsync = std::numeric_limits<int32>::max();
+static FAutoConsoleVariableRef CVarD3D12SubmissionMaxExecuteBatchSizeAsync(
+	TEXT("r.D3D12.Submission.MaxExecuteBatchSize.Async"),
+	GD3D12SubmissionMaxExecuteBatchSizeAsync,
+	TEXT("The maximum number of command lists to pass to a single ExecuteCommandLists invocation for async queues\n")
+	TEXT("The valid range is 1 to INT_MAX inclusive. Values less than 1 will be clamped to 1"),
+	ECVF_RenderThreadSafe);
+
 static std::atomic<int> GGPUCrashDetected = false;
 
 class FD3D12Thread final : private FRunnable
@@ -331,13 +355,21 @@ void FD3D12DynamicRHI::SubmitPayloads(TArray<FD3D12Payload*>&& Payloads)
 	ProcessInterruptQueueUntil(nullptr);
 }
 
-static int32 GetMaxExecuteBatchSize()
+static int32 GetMaxExecuteBatchSize(ED3D12QueueType QueueType)
 {
-	return
-#if UE_BUILD_DEBUG
-		GRHIGlobals.IsDebugLayerEnabled ? 1 :
-#endif
-		TNumericLimits<int32>::Max();
+	switch (QueueType)
+	{
+		case ED3D12QueueType::Direct:
+			return std::max(1, GD3D12SubmissionMaxExecuteBatchSizeDirect);
+		case ED3D12QueueType::Copy:
+			return std::max(1, GD3D12SubmissionMaxExecuteBatchSizeCopy);
+		case ED3D12QueueType::Async:
+			return std::max(1, GD3D12SubmissionMaxExecuteBatchSizeAsync);
+		default:
+			// Need to add new queue type and CVar
+			checkNoEntry();
+			return std::numeric_limits<int32>::max();
+	}
 }
 
 FD3D12DynamicRHI::FProcessResult FD3D12DynamicRHI::ProcessSubmissionQueue()
@@ -441,10 +473,10 @@ FD3D12DynamicRHI::FProcessResult FD3D12DynamicRHI::ProcessSubmissionQueue()
 					// before async compute work), so we gather potential work across all
 					// queues for this device.
 					//
-					const uint32 MaxBatchSize = GetMaxExecuteBatchSize();
 					auto AccumulateQueries = [&](FD3D12CommandList* CommandList)
 					{
 						FD3D12Queue& TargetQueue = CommandList->Device->GetQueue(CommandList->QueueType);
+						const uint32 MaxBatchSize = GetMaxExecuteBatchSize(TargetQueue.QueueType);
 
 						// Occlusion + Pipeline Stats Queries
 						TargetQueue.BatchedObjects.OcclusionQueries.Append(MoveTemp(CommandList->State.OcclusionQueries));
@@ -759,7 +791,7 @@ void FD3D12DynamicRHI::FlushBatchedPayloads(FD3D12Queue::FPayloadArray& Payloads
 			CommandLists.Append(MoveTemp(Payload->CommandListsToExecute));
 		}
 
-		const int32 MaxBatchSize = GetMaxExecuteBatchSize();
+		const int32 MaxBatchSize = GetMaxExecuteBatchSize(Queue.QueueType);
 		const int32 NumCommandLists = D3DCommandLists.Num();
 
 		for (int32 DispatchNum, Offset = 0; Offset < NumCommandLists; Offset += DispatchNum)
diff --git a/Engine/Source/Runtime/D3D12RHI/Private/Windows/WindowsD3D12Device.cpp b/Engine/Source/Runtime/D3D12RHI/Private/Windows/WindowsD3D12Device.cpp
--- a/Engine/Source/Runtime/D3D12RHI/Private/Windows/WindowsD3D12Device.cpp
+++ b/Engine/Source/Runtime/D3D12RHI/Private/Windows/WindowsD3D12Device.cpp
@@ -1547,6 +1547,36 @@ void FD3D12DynamicRHI::Init()
 	} // if NVIDIA
 #endif // NVAPI
 
+	if (IsRHIDeviceNVIDIA())
+	{
+		//
+		// See Jira UE-233616
+		//
+		// The NVIDIA driver drops barriers at the beginning of command lists on the async queue because
+		// it assumes each list is executed within its own ECL scope and therefore the GPU is idle'd between
+		// each. That's not the case if we execute more than one list at a time, so break it up into individual
+		// executions to make it true.
+		//
+		// @TODO - Add driver version check and message text once a fix is released. Also raise message priority to warning.
+		//
+
+		static constexpr auto D3D12SubmissionExecuteBatchSizeAsyncCVarName = TEXT("r.D3D12.Submission.MaxExecuteBatchSize.Async");
+		const auto CVarD3D12SubmissionAsyncMaxExecuteBatchSize = IConsoleManager::Get().FindConsoleVariable(D3D12SubmissionExecuteBatchSizeAsyncCVarName);
+		if (CVarD3D12SubmissionAsyncMaxExecuteBatchSize)
+		{
+			CVarD3D12SubmissionAsyncMaxExecuteBatchSize->Set(1);
+			UE_LOG(LogD3D12RHI, Display,
+				TEXT("Batched command list execution is disabled for async queues due to known bugs in the current driver."));
+		}
+		else
+		{
+			// If we're here, someone moved our cheese and this work-around logic needs to be updated
+			UE_LOG(LogD3D12RHI, Error,
+				TEXT("Couldn't find CVar %s to work around a known bug in the current driver."), D3D12SubmissionExecuteBatchSizeAsyncCVarName);
+			checkNoEntry();
+		}
+	} // if NVIDIA
+
 #if WITH_AMD_AGS
 	if (GRHISupportsRayTracing
 		&& IsRHIDeviceAMD()
