DX12 Barriers

- Proper Discard of Buffers and Textures that can be effectively discarded on the Async Pipe

- UAV -> SRVMask transition that are Async pipe -> All Pipes are now split in two: UAV -> SRVCompute on Async and SRVCompute -> SRVMask on GFx with a Fence between them

- Disable, when Async Compute is on, the D3D12_MESSAGE_ID_RESOURCE_BARRIER_BEFORE_AFTER_MISMATCH d3ddebug error message due to the split transition on the Aysnc and Gfx Pipe, on the GPU the state is correct but the validation layer cannot know on the CPU

- Remove extra ERHITransitionCreateFlags::NoFence flag added inside  FRHICommandListBase::TransitionInternal

- Clean up old code comments

diff --git a/Engine/Source/Runtime/D3D12RHI/Private/D3D12Adapter.cpp b/Engine/Source/Runtime/D3D12RHI/Private/D3D12Adapter.cpp
--- a/Engine/Source/Runtime/D3D12RHI/Private/D3D12Adapter.cpp
+++ b/Engine/Source/Runtime/D3D12RHI/Private/D3D12Adapter.cpp
@@ -650,152 +650,157 @@ void FD3D12Adapter::CreateRootDevice(bool bWithDebug)
 	}
 #endif
 
-#if !(UE_BUILD_SHIPPING && WITH_EDITOR) && D3D12_SUPPORTS_INFO_QUEUE
-	// Add some filter outs for known debug spew messages (that we don't care about)
-	if (bWithDebug)
+#if WITH_MGPU
+	GNumExplicitGPUsForRendering = 1;
+	if (Desc.NumDeviceNodes > 1)
+	{
+		GNumExplicitGPUsForRendering = Desc.NumDeviceNodes;
+		UE_LOG(LogD3D12RHI, Log, TEXT("Enabling multi-GPU with %d nodes"), Desc.NumDeviceNodes);
+	}
+#endif
+
+	if (RootDevice)
 	{
-		ID3D12InfoQueue *pd3dInfoQueue = nullptr;
-		VERIFYD3D12RESULT(RootDevice->QueryInterface(__uuidof(ID3D12InfoQueue), (void**)&pd3dInfoQueue));
-		if (pd3dInfoQueue)
+		extern int32 GD3D12RHIStablePowerState;
+		if (GD3D12RHIStablePowerState == 2)
 		{
-			D3D12_INFO_QUEUE_FILTER NewFilter;
-			FMemory::Memzero(&NewFilter, sizeof(NewFilter));
-
-			// Turn off info msgs as these get really spewy
-			const bool bLogWarnings = D3D12_ShouldBreakOnD3DDebugWarnings() || D3D12_ShouldLogD3DDebugWarnings();
-			D3D12_MESSAGE_SEVERITY DenySeverity[] = { D3D12_MESSAGE_SEVERITY_INFO, D3D12_MESSAGE_SEVERITY_WARNING };
-			NewFilter.DenyList.NumSeverities = 1 + (bLogWarnings ? 0 : 1);
-			NewFilter.DenyList.pSeverityList = DenySeverity;
-			// Be sure to carefully comment the reason for any additions here!  Someone should be able to look at it later and get an idea of whether it is still necessary.
-			TArray<D3D12_MESSAGE_ID, TInlineAllocator<16>> DenyIds = {
-
-				// The Pixel Shader expects a Render Target View bound to slot 0, but the PSO indicates that none will be bound.
-				// This typically happens when a non-depth-only pixel shader is used for depth-only rendering.
-				D3D12_MESSAGE_ID_CREATEGRAPHICSPIPELINESTATE_RENDERTARGETVIEW_NOT_SET,
-
-				// QUERY_BEGIN_ABANDONING_PREVIOUS_RESULTS - The RHI exposes the interface to make and issue queries and a separate interface to use that data.
-				//		Currently there is a situation where queries are issued and the results may be ignored on purpose.  Filtering out this message so it doesn't
-				//		swarm the debug spew and mask other important warnings
-				//D3D12_MESSAGE_ID_QUERY_BEGIN_ABANDONING_PREVIOUS_RESULTS,
-				//D3D12_MESSAGE_ID_QUERY_END_ABANDONING_PREVIOUS_RESULTS,
-
-				// D3D12_MESSAGE_ID_CREATEINPUTLAYOUT_EMPTY_LAYOUT - This is a warning that gets triggered if you use a null vertex declaration,
-				//       which we want to do when the vertex shader is generating vertices based on ID.
-				D3D12_MESSAGE_ID_CREATEINPUTLAYOUT_EMPTY_LAYOUT,
-
-				// D3D12_MESSAGE_ID_COMMAND_LIST_DRAW_INDEX_BUFFER_TOO_SMALL - This warning gets triggered by Slate draws which are actually using a valid index range.
-				//		The invalid warning seems to only happen when VS 2012 is installed.  Reported to MS.  
-				//		There is now an assert in DrawIndexedPrimitive to catch any valid errors reading from the index buffer outside of range.
-				D3D12_MESSAGE_ID_COMMAND_LIST_DRAW_INDEX_BUFFER_TOO_SMALL,
-
-				// D3D12_MESSAGE_ID_DEVICE_DRAW_RENDERTARGETVIEW_NOT_SET - This warning gets triggered by shadow depth rendering because the shader outputs
-				//		a color but we don't bind a color render target. That is safe as writes to unbound render targets are discarded.
-				//		Also, batched elements triggers it when rendering outside of scene rendering as it outputs to the GBuffer containing normals which is not bound.
-				//(D3D12_MESSAGE_ID)3146081, // D3D12_MESSAGE_ID_DEVICE_DRAW_RENDERTARGETVIEW_NOT_SET,
-				// BUGBUG: There is a D3D12_MESSAGE_ID_DEVICE_DRAW_DEPTHSTENCILVIEW_NOT_SET, why not one for RT?
-
-				// D3D12_MESSAGE_ID_CLEARDEPTHSTENCILVIEW_MISMATCHINGCLEARVALUE/D3D12_MESSAGE_ID_CLEARRENDERTARGETVIEW_MISMATCHINGCLEARVALUE - 
-				//      This warning gets triggered by ClearDepthStencilView/ClearRenderTargetView because when the resource was created
-				//      it wasn't passed an optimized clear color (see CreateCommitedResource). This shows up a lot and is very noisy.
-				D3D12_MESSAGE_ID_CLEARRENDERTARGETVIEW_MISMATCHINGCLEARVALUE,
-				D3D12_MESSAGE_ID_CLEARDEPTHSTENCILVIEW_MISMATCHINGCLEARVALUE,
-
-				// D3D12_MESSAGE_ID_EXECUTECOMMANDLISTS_GPU_WRITTEN_READBACK_RESOURCE_MAPPED - This warning gets triggered by ExecuteCommandLists.
-				//		if it contains a readback resource that still has mapped subresources when executing a command list that performs a copy operation to the resource.
-				//		This may be ok if any data read from the readback resources was flushed by calling Unmap() after the resourcecopy operation completed.
-				//		We intentionally keep the readback resources persistently mapped.
-				D3D12_MESSAGE_ID_EXECUTECOMMANDLISTS_GPU_WRITTEN_READBACK_RESOURCE_MAPPED,
-
-				// This shows up a lot and is very noisy. It would require changes to the resource tracking system
-				// but will hopefully be resolved when the RHI switches to use the engine's resource tracking system.
-				D3D12_MESSAGE_ID_RESOURCE_BARRIER_DUPLICATE_SUBRESOURCE_TRANSITIONS,
-
-				// This error gets generated on the first run when you install a new driver. The code handles this error properly and resets the PipelineLibrary,
-				// so we can safely ignore this message. It could possibly be avoided by adding driver version to the PSO cache filename, but an average user is unlikely
-				// to be interested in keeping PSO caches associated with old drivers around on disk, so it's better to just reset.
-				D3D12_MESSAGE_ID_CREATEPIPELINELIBRARY_DRIVERVERSIONMISMATCH,
-
-				// D3D complain about overlapping GPU addresses when aliasing DataBuffers in the same command list when using the Transient Allocator - it looks like
-				// it ignored the aliasing barriers to validate, and probably can't check them when called from IASetVertexBuffers because it only has GPU Virtual Addresses then
-				D3D12_MESSAGE_ID_HEAP_ADDRESS_RANGE_INTERSECTS_MULTIPLE_BUFFERS,
-
-				// Ignore draw vertex buffer not set or too small - these are warnings and if the shader doesn't read from it it's fine. This happens because vertex
-				// buffers are not removed from the cache, but only get removed when another buffer is set at the same slot or when the buffer gets destroyed.
-				D3D12_MESSAGE_ID_COMMAND_LIST_DRAW_VERTEX_BUFFER_NOT_SET,
-				D3D12_MESSAGE_ID_COMMAND_LIST_DRAW_VERTEX_BUFFER_TOO_SMALL,
-
-				// D3D12 complains when a buffer is created with a specific initial resource state while all buffers are currently created in COMMON state. The 
-				// next transition is then done use state promotion. It's just a warning and we need to keep track of the correct initial state as well for upcoming
-				// internal transitions.
-				D3D12_MESSAGE_ID_CREATERESOURCE_STATE_IGNORED,
+			bool bWorked = SUCCEEDED(RootDevice->SetStablePowerState(true));
+			// This will fail if windows developper mode is not enabled. Windows will remove the adapter on failure, so we can't really gracefully exit here.
+			checkf(bWorked, TEXT("Enabling state power state requires Windows developer mode to be enabled."));
+		}
+	}
+}
+
+// Add some filter outs for known debug spew messages (that we don't care about)
+void FD3D12Adapter::CreateD3DInfoQueue()
+{
+#if !(UE_BUILD_SHIPPING && WITH_EDITOR) && D3D12_SUPPORTS_INFO_QUEUE
+	ID3D12InfoQueue* pd3dInfoQueue = nullptr;
+	VERIFYD3D12RESULT(RootDevice->QueryInterface(__uuidof(ID3D12InfoQueue), (void**)&pd3dInfoQueue));
+	if (pd3dInfoQueue)
+	{
+		D3D12_INFO_QUEUE_FILTER NewFilter;
+		FMemory::Memzero(&NewFilter, sizeof(NewFilter));
+
+		// Turn off info msgs as these get really spewy
+		const bool bLogWarnings = D3D12_ShouldBreakOnD3DDebugWarnings() || D3D12_ShouldLogD3DDebugWarnings();
+		D3D12_MESSAGE_SEVERITY DenySeverity[] = { D3D12_MESSAGE_SEVERITY_INFO, D3D12_MESSAGE_SEVERITY_WARNING };
+		NewFilter.DenyList.NumSeverities = 1 + (bLogWarnings ? 0 : 1);
+		NewFilter.DenyList.pSeverityList = DenySeverity;
+		// Be sure to carefully comment the reason for any additions here!  Someone should be able to look at it later and get an idea of whether it is still necessary.
+		TArray<D3D12_MESSAGE_ID, TInlineAllocator<16>> DenyIds = {
+
+			// The Pixel Shader expects a Render Target View bound to slot 0, but the PSO indicates that none will be bound.
+			// This typically happens when a non-depth-only pixel shader is used for depth-only rendering.
+			D3D12_MESSAGE_ID_CREATEGRAPHICSPIPELINESTATE_RENDERTARGETVIEW_NOT_SET,
+
+			// QUERY_BEGIN_ABANDONING_PREVIOUS_RESULTS - The RHI exposes the interface to make and issue queries and a separate interface to use that data.
+			//		Currently there is a situation where queries are issued and the results may be ignored on purpose.  Filtering out this message so it doesn't
+			//		swarm the debug spew and mask other important warnings
+			//D3D12_MESSAGE_ID_QUERY_BEGIN_ABANDONING_PREVIOUS_RESULTS,
+			//D3D12_MESSAGE_ID_QUERY_END_ABANDONING_PREVIOUS_RESULTS,
+
+			// D3D12_MESSAGE_ID_CREATEINPUTLAYOUT_EMPTY_LAYOUT - This is a warning that gets triggered if you use a null vertex declaration,
+			//       which we want to do when the vertex shader is generating vertices based on ID.
+			D3D12_MESSAGE_ID_CREATEINPUTLAYOUT_EMPTY_LAYOUT,
+
+			// D3D12_MESSAGE_ID_COMMAND_LIST_DRAW_INDEX_BUFFER_TOO_SMALL - This warning gets triggered by Slate draws which are actually using a valid index range.
+			//		The invalid warning seems to only happen when VS 2012 is installed.  Reported to MS.  
+			//		There is now an assert in DrawIndexedPrimitive to catch any valid errors reading from the index buffer outside of range.
+			//D3D12_MESSAGE_ID_COMMAND_LIST_DRAW_INDEX_BUFFER_TOO_SMALL,
+
+			// D3D12_MESSAGE_ID_DEVICE_DRAW_RENDERTARGETVIEW_NOT_SET - This warning gets triggered by shadow depth rendering because the shader outputs
+			//		a color but we don't bind a color render target. That is safe as writes to unbound render targets are discarded.
+			//		Also, batched elements triggers it when rendering outside of scene rendering as it outputs to the GBuffer containing normals which is not bound.
+			//(D3D12_MESSAGE_ID)3146081, // D3D12_MESSAGE_ID_DEVICE_DRAW_RENDERTARGETVIEW_NOT_SET,
+			// BUGBUG: There is a D3D12_MESSAGE_ID_DEVICE_DRAW_DEPTHSTENCILVIEW_NOT_SET, why not one for RT?
+
+			// D3D12_MESSAGE_ID_CLEARDEPTHSTENCILVIEW_MISMATCHINGCLEARVALUE/D3D12_MESSAGE_ID_CLEARRENDERTARGETVIEW_MISMATCHINGCLEARVALUE - 
+			//      This warning gets triggered by ClearDepthStencilView/ClearRenderTargetView because when the resource was created
+			//      it wasn't passed an optimized clear color (see CreateCommitedResource). This shows up a lot and is very noisy.
+			D3D12_MESSAGE_ID_CLEARRENDERTARGETVIEW_MISMATCHINGCLEARVALUE,
+			D3D12_MESSAGE_ID_CLEARDEPTHSTENCILVIEW_MISMATCHINGCLEARVALUE,
+
+			// D3D12_MESSAGE_ID_EXECUTECOMMANDLISTS_GPU_WRITTEN_READBACK_RESOURCE_MAPPED - This warning gets triggered by ExecuteCommandLists.
+			//		if it contains a readback resource that still has mapped subresources when executing a command list that performs a copy operation to the resource.
+			//		This may be ok if any data read from the readback resources was flushed by calling Unmap() after the resourcecopy operation completed.
+			//		We intentionally keep the readback resources persistently mapped.
+			D3D12_MESSAGE_ID_EXECUTECOMMANDLISTS_GPU_WRITTEN_READBACK_RESOURCE_MAPPED,
+
+			// This error gets generated on the first run when you install a new driver. The code handles this error properly and resets the PipelineLibrary,
+			// so we can safely ignore this message. It could possibly be avoided by adding driver version to the PSO cache filename, but an average user is unlikely
+			// to be interested in keeping PSO caches associated with old drivers around on disk, so it's better to just reset.
+			D3D12_MESSAGE_ID_CREATEPIPELINELIBRARY_DRIVERVERSIONMISMATCH,
+
+			// D3D complain about overlapping GPU addresses when aliasing DataBuffers in the same command list when using the Transient Allocator - it looks like
+			// it ignored the aliasing barriers to validate, and probably can't check them when called from IASetVertexBuffers because it only has GPU Virtual Addresses then
+			D3D12_MESSAGE_ID_HEAP_ADDRESS_RANGE_INTERSECTS_MULTIPLE_BUFFERS,
+
+			// Ignore draw vertex buffer not set or too small - these are warnings and if the shader doesn't read from it it's fine. This happens because vertex
+			// buffers are not removed from the cache, but only get removed when another buffer is set at the same slot or when the buffer gets destroyed.
+			D3D12_MESSAGE_ID_COMMAND_LIST_DRAW_VERTEX_BUFFER_NOT_SET,
+			D3D12_MESSAGE_ID_COMMAND_LIST_DRAW_VERTEX_BUFFER_TOO_SMALL,
+
+			// D3D12 complains when a buffer is created with a specific initial resource state while all buffers are currently created in COMMON state. The 
+			// next transition is then done use state promotion. It's just a warning and we need to keep track of the correct initial state as well for upcoming
+			// internal transitions.
+			D3D12_MESSAGE_ID_CREATERESOURCE_STATE_IGNORED,
 
 #if ENABLE_RESIDENCY_MANAGEMENT
-				// TODO: Remove this when the debug layers work for executions which are guarded by a fence
-				D3D12_MESSAGE_ID_INVALID_USE_OF_NON_RESIDENT_RESOURCE,
+			// TODO: Remove this when the debug layers work for executions which are guarded by a fence
+			D3D12_MESSAGE_ID_INVALID_USE_OF_NON_RESIDENT_RESOURCE,
 #endif
 
-				// When optimizing the graph DirectML tries various configurations to check if meta command can be used.
-				// If a particular configuration is not supported the debug log message will be displayed.
-				// For large DirectML graphs there can be dozens of messages like this.
-				// This is safe to ignore
-				D3D12_MESSAGE_ID_META_COMMAND_UNSUPPORTED_PARAMS,
+			// When optimizing the graph DirectML tries various configurations to check if meta command can be used.
+			// If a particular configuration is not supported the debug log message will be displayed.
+			// For large DirectML graphs there can be dozens of messages like this.
+			// This is safe to ignore
+			D3D12_MESSAGE_ID_META_COMMAND_UNSUPPORTED_PARAMS,
 
-				// D3D Agility SDK bug (version 614) where scratch allocation memory for BuildRaytracingAccelerationStructure is always computed using build size
-				// while operation is update and needs less memory (verified by MS and will be fixed in the next Agility SDK)
-				// See: UE-222685 to remove again from Deny list after Agility SDK upgrade
-				D3D12_MESSAGE_ID_HEAP_ADDRESS_RANGE_HAS_NO_RESOURCE,
+			// D3D Agility SDK bug (version 614) where scratch allocation memory for BuildRaytracingAccelerationStructure is always computed using build size
+			// while operation is update and needs less memory (verified by MS and will be fixed in the next Agility SDK)
+			// See: UE-222685 to remove again from Deny list after Agility SDK upgrade
+			D3D12_MESSAGE_ID_HEAP_ADDRESS_RANGE_HAS_NO_RESOURCE,
 
-			};
+		};
 
 #if PLATFORM_DESKTOP
-			if (!FWindowsPlatformMisc::VerifyWindowsVersion(10, 0, 18363))
-			{
-				// Ignore a known false positive error due to a bug in validation layer in certain older Windows versions
-				DenyIds.Add(D3D12_MESSAGE_ID_COPY_DESCRIPTORS_INVALID_RANGES);
-			}
+		if (!FWindowsPlatformMisc::VerifyWindowsVersion(10, 0, 18363))
+		{
+			// Ignore a known false positive error due to a bug in validation layer in certain older Windows versions
+			DenyIds.Add(D3D12_MESSAGE_ID_COPY_DESCRIPTORS_INVALID_RANGES);
+		}
 #endif // PLATFORM_DESKTOP
 
-			NewFilter.DenyList.NumIDs = DenyIds.Num();
-			NewFilter.DenyList.pIDList = DenyIds.GetData();
+		if (GSupportsEfficientAsyncCompute)
+		{
+			// This is not optimal to have it disable but this should be absolutely removed when DX12 can use EB 
+			// The reason why this is here is because when async compute is on we can receive from RDG a transition UAV -> SRVMask on the Async Pipe
+			// DX12 cannot transition to Gfx states on the Async Pipe therefore we do: UAV->SRVCompute on Async, Fence Async -> Gfx, SRVCompute->SRVMask on Gfx
+			// Since the debug layer is on the CPU and the submission is happening in parallel threads the validation layer cannot known that the order on the GPU will be UAV->SRVCompute->SRVMask
+			DenyIds.Add(D3D12_MESSAGE_ID_RESOURCE_BARRIER_BEFORE_AFTER_MISMATCH);
+		}
 
-			pd3dInfoQueue->PushStorageFilter(&NewFilter);
+		NewFilter.DenyList.NumIDs = DenyIds.Num();
+		NewFilter.DenyList.pIDList = DenyIds.GetData();
 
-			// Break on D3D debug errors.
-			pd3dInfoQueue->SetBreakOnSeverity(D3D12_MESSAGE_SEVERITY_ERROR, D3D12_ShouldBreakOnD3DDebugErrors());
+		pd3dInfoQueue->PushStorageFilter(&NewFilter);
 
-			// Enable this to break on a specific id in order to quickly get a callstack
-			//pd3dInfoQueue->SetBreakOnID(D3D12_MESSAGE_ID_DEVICE_DRAW_CONSTANT_BUFFER_TOO_SMALL, true);
+		// Break on D3D debug errors.
+		pd3dInfoQueue->SetBreakOnSeverity(D3D12_MESSAGE_SEVERITY_ERROR, D3D12_ShouldBreakOnD3DDebugErrors());
 
-			// Break on D3D warnings if warning log or warning breakpoint is enabled
-			if (bLogWarnings)
-			{
-				pd3dInfoQueue->SetBreakOnSeverity(D3D12_MESSAGE_SEVERITY_WARNING, D3D12_ShouldBreakOnD3DDebugWarnings());
-			}
+		// Enable this to break on a specific id in order to quickly get a callstack
+		//pd3dInfoQueue->SetBreakOnID(D3D12_MESSAGE_ID_DEVICE_DRAW_CONSTANT_BUFFER_TOO_SMALL, true);
 
-			pd3dInfoQueue->Release();
+		// Break on D3D warnings if warning log or warning breakpoint is enabled
+		if (bLogWarnings)
+		{
+			pd3dInfoQueue->SetBreakOnSeverity(D3D12_MESSAGE_SEVERITY_WARNING, D3D12_ShouldBreakOnD3DDebugWarnings());
 		}
-	}
-#endif
 
-#if WITH_MGPU
-	GNumExplicitGPUsForRendering = 1;
-	if (Desc.NumDeviceNodes > 1)
-	{
-		GNumExplicitGPUsForRendering = Desc.NumDeviceNodes;
-		UE_LOG(LogD3D12RHI, Log, TEXT("Enabling multi-GPU with %d nodes"), Desc.NumDeviceNodes);
+		pd3dInfoQueue->Release();
 	}
 #endif
-
-	if (RootDevice)
-	{
-		extern int32 GD3D12RHIStablePowerState;
-		if (GD3D12RHIStablePowerState == 2)
-		{
-			bool bWorked = SUCCEEDED(RootDevice->SetStablePowerState(true));
-			// This will fail if windows developper mode is not enabled. Windows will remove the adapter on failure, so we can't really gracefully exit here.
-			checkf(bWorked, TEXT("Enabling state power state requires Windows developer mode to be enabled."));
-		}
-	}
 }
 
 void FD3D12Adapter::SetDrawingViewport(FD3D12Viewport* InViewport)
@@ -1171,6 +1176,11 @@ void FD3D12Adapter::InitializeDevices()
 		GSupportsEfficientAsyncCompute = false;
 #endif
 
+		if (bWithD3DDebug)
+		{
+			CreateD3DInfoQueue();
+		}
+
 		if (!FPlatformMemory::SupportsFastVRAMMemory() && GSupportsEfficientAsyncCompute)
 		{
 			GRHIGlobals.SupportsAsyncComputeTransientAliasing = true;
diff --git a/Engine/Source/Runtime/D3D12RHI/Private/D3D12Adapter.h b/Engine/Source/Runtime/D3D12RHI/Private/D3D12Adapter.h
--- a/Engine/Source/Runtime/D3D12RHI/Private/D3D12Adapter.h
+++ b/Engine/Source/Runtime/D3D12RHI/Private/D3D12Adapter.h
@@ -461,6 +461,8 @@ protected:
 
 	virtual void CreateRootDevice(bool bWithDebug);
 
+	void CreateD3DInfoQueue();
+
 	virtual void AllocateBuffer(FD3D12Device* Device,
 		const D3D12_RESOURCE_DESC& Desc,
 		uint32 Size,
diff --git a/Engine/Source/Runtime/D3D12RHI/Private/D3D12CommandContext.cpp b/Engine/Source/Runtime/D3D12RHI/Private/D3D12CommandContext.cpp
--- a/Engine/Source/Runtime/D3D12RHI/Private/D3D12CommandContext.cpp
+++ b/Engine/Source/Runtime/D3D12RHI/Private/D3D12CommandContext.cpp
@@ -797,22 +797,21 @@ void FD3D12DynamicRHI::RHICreateTransition(FRHITransition* Transition, const FRH
 	Data->bCrossPipeline = bCrossPipeline;
 
 	// In DX12 we cannot perform resource barrier with graphics state on the AsyncCompute pipe
-	// this check is here to be able to force a crosspipe transition coming from AsyncCompute with graphics states to be processed in the Graphics pipe
+	// This check is here to be able to force a crosspipe transition coming from AsyncCompute with graphics states to be split and processed in the both the Async and Graphics pipe
+	// This case can be removed when using EB on DX12 
 	if (bAsyncToAllPipelines)
 	{
 		for (const FRHITransitionInfo& TransitionInfo : CreateInfo.TransitionInfos)
 		{
-			if (EnumHasAnyFlags(TransitionInfo.AccessAfter, ERHIAccess::SRVGraphics | ERHIAccess::UAVGraphics))
+			if (EnumHasAnyFlags(TransitionInfo.AccessAfter, ERHIAccess::SRVGraphics))
 			{
 				Data->bAsyncToAllPipelines = true;
-				// note that we are ignoring the ERHITransitionCreateFlags::NoFence flag for this specific case and forcing the fences creation, since we need fences between the pipes
-				Data->bCrossPipeline = true;
 				break;
 			}
 		}
 	}
 
-	if (Data->bCrossPipeline)
+	if ((Data->bCrossPipeline) || (Data->bAsyncToAllPipelines))
 	{
 		// Create one sync point per device, per source pipe
 		for (uint32 Index : FRHIGPUMask::All())
diff --git a/Engine/Source/Runtime/D3D12RHI/Private/D3D12CommandList.cpp b/Engine/Source/Runtime/D3D12RHI/Private/D3D12CommandList.cpp
--- a/Engine/Source/Runtime/D3D12RHI/Private/D3D12CommandList.cpp
+++ b/Engine/Source/Runtime/D3D12RHI/Private/D3D12CommandList.cpp
@@ -474,19 +474,6 @@ static inline bool IsTransitionNeeded(D3D12_RESOURCE_STATES Before, D3D12_RESOUR
 	check(Before != D3D12_RESOURCE_STATE_CORRUPT && After != D3D12_RESOURCE_STATE_CORRUPT);
 	check(Before != D3D12_RESOURCE_STATE_TBD && After != D3D12_RESOURCE_STATE_TBD);
 
-	// Todo: verify that high level is taking care of this optimization
-// 	// Depth write is actually a suitable for read operations as a "normal" depth buffer.
-// 	if (Before == D3D12_RESOURCE_STATE_DEPTH_WRITE && After == D3D12_RESOURCE_STATE_DEPTH_READ)
-// 	{
-// 		{
-// 			if (GD3D12ExtraDepthTransitions)
-// 			{
-// 				After = D3D12_RESOURCE_STATE_DEPTH_WRITE;
-// 			}
-// 			return false;
-// 		}
-// 	}
-
 	// COMMON is an oddball state that doesn't follow the RESOURE_STATE pattern of 
 	// having exactly one bit set so we need to special case these
 	if (After == D3D12_RESOURCE_STATE_COMMON)
@@ -496,18 +483,6 @@ static inline bool IsTransitionNeeded(D3D12_RESOURCE_STATES Before, D3D12_RESOUR
 		return true;
 	}
 
-	// Todo: verify that high level is taking care of this optimization
-	// We should avoid doing read-to-read state transitions. But when we do, we should avoid turning off already transitioned bits,
-	// e.g. VERTEX_BUFFER -> SHADER_RESOURCE is turned into VERTEX_BUFFER -> VERTEX_BUFFER | SHADER_RESOURCE.
-	// This reduces the number of resource transitions and ensures automatic states from resource bindings get properly combined.
-// 	D3D12_RESOURCE_STATES Combined = Before | After;
-// 	if ((Combined & (D3D12_RESOURCE_STATE_GENERIC_READ | D3D12_RESOURCE_STATE_INDIRECT_ARGUMENT)) == Combined)
-// 	{
-// 		{
-// 			After = Combined;
-// 		}
-// 	}
-
 	return Before != After;
 }
 
diff --git a/Engine/Source/Runtime/D3D12RHI/Private/D3D12Commands.cpp b/Engine/Source/Runtime/D3D12RHI/Private/D3D12Commands.cpp
--- a/Engine/Source/Runtime/D3D12RHI/Private/D3D12Commands.cpp
+++ b/Engine/Source/Runtime/D3D12RHI/Private/D3D12Commands.cpp
@@ -228,8 +228,8 @@ void EnumerateSubresources(FD3D12Resource* Resource, const FRHITransitionInfo& I
 	}
 }
 
-template <typename FunctionType>
-void ProcessResource(FD3D12CommandContext& Context, const FRHITransitionInfo& Info, FunctionType Function)
+static
+std::tuple<FD3D12Resource*,FD3D12Texture*> GetResourceAndTexture(FD3D12CommandContext& Context, const FRHITransitionInfo& Info)
 {
 	switch (Info.Type)
 	{
@@ -237,9 +237,11 @@ void ProcessResource(FD3D12CommandContext& Context, const FRHITransitionInfo& In
 	{
 		FD3D12UnorderedAccessView* UAV = Context.RetrieveObject<FD3D12UnorderedAccessView_RHI>(Info.UAV);
 		check(UAV);
-
-		Function(Info, UAV->GetResource(), nullptr, nullptr);
-		break;
+		if (UAV)
+		{
+			return { UAV->GetResource(), nullptr };
+		}
+		return { nullptr, nullptr };
 	}
 	case FRHITransitionInfo::EType::Buffer:
 	{
@@ -248,9 +250,9 @@ void ProcessResource(FD3D12CommandContext& Context, const FRHITransitionInfo& In
 		check(Buffer || GNumExplicitGPUsForRendering > 1);
 		if (Buffer)
 		{
-			Function(Info, Buffer->GetResource(), nullptr, nullptr);
+			return { Buffer->GetResource(), nullptr };
 		}
-		break;
+		return { nullptr, nullptr };
 	}
 	case FRHITransitionInfo::EType::Texture:
 	{
@@ -259,28 +261,35 @@ void ProcessResource(FD3D12CommandContext& Context, const FRHITransitionInfo& In
 		check(Texture || GNumExplicitGPUsForRendering > 1);
 		if (Texture)
 		{
-			FD3D12Texture* DiscardTextureOut = nullptr;
-#if PLATFORM_REQUIRES_TYPELESS_RESOURCE_DISCARD_WORKAROUND
-			if (Texture->GetRequiresTypelessResourceDiscardWorkaround())
-			{
-				DiscardTextureOut = Texture;
-			}
-#endif // #if PLATFORM_REQUIRES_TYPELESS_RESOURCE_DISCARD_WORKAROUND
-			Function(Info, Texture->GetResource(), Texture, DiscardTextureOut);
+			return { Texture->GetResource(), Texture };
 		}
-		break;
+		return { nullptr, nullptr };
 	}
 	default:
 		checkNoEntry();
+		return { nullptr, nullptr };
 		break;
 	}
 }
 
-// static bool ProcessTransitionDuringBegin(const FD3D12TransitionData* Data)
-// {
-// 	// Pipe changes which are not ending with graphics or targeting all pipelines are handle during begin
-// 	return !EnumHasAllFlags(Data->SrcPipelines, ERHIPipeline::All) && ((Data->SrcPipelines != Data->DstPipelines && Data->DstPipelines != ERHIPipeline::Graphics) || EnumHasAllFlags(Data->DstPipelines, ERHIPipeline::All));
-// }
+template <typename FunctionType>
+void ProcessResource(FD3D12CommandContext& Context, const FRHITransitionInfo& Info, FunctionType Function)
+{
+	auto [Resource, Texture] = GetResourceAndTexture(Context, Info);
+	FD3D12Texture* DiscardTextureOut = nullptr;
+
+	if ((Info.Type == FRHITransitionInfo::EType::Texture) && Texture)
+	{
+#if PLATFORM_REQUIRES_TYPELESS_RESOURCE_DISCARD_WORKAROUND
+		if (Texture->GetRequiresTypelessResourceDiscardWorkaround())
+		{
+			DiscardTextureOut = Texture;
+		}
+#endif // #if PLATFORM_REQUIRES_TYPELESS_RESOURCE_DISCARD_WORKAROUND
+	}
+
+	Function(Info, Resource, Texture, DiscardTextureOut);
+}
 
 // Pipe changes which are not ending with graphics or targeting all pipelines are handle during begin
 static bool ProcessTransitionDuringBegin(const FD3D12TransitionData* Data)
@@ -309,6 +318,11 @@ static bool ShouldProcessTransition(const FD3D12TransitionData* Data, bool bIsBe
 		{
 			return true;
 		}
+		
+		if ((bIsBeginTransition) && (ExecutingPipeline == ERHIPipeline::AsyncCompute))
+		{
+			return true;
+		}
 
 		return false;
 	}
@@ -316,6 +330,22 @@ static bool ShouldProcessTransition(const FD3D12TransitionData* Data, bool bIsBe
 	return (bIsBeginTransition ? ProcessTransitionDuringBegin(Data) : !ProcessTransitionDuringBegin(Data));
 }
 
+static bool IsImpossibleToDiscardOnAsync(
+	FD3D12CommandContext& Context,
+	const FRHITransitionInfo& Info)
+{
+	const ED3D12QueueType QueueType = Context.QueueType;
+	if (QueueType == ED3D12QueueType::Async)
+	{
+		const auto [Resource, Texture] = GetResourceAndTexture(Context, Info);
+
+		// It is expected that this will return COMMON if the initial state isn't compatible with the queue type
+		const D3D12_RESOURCE_STATES InitialState = GetInitialResourceState(Resource->GetDesc(), QueueType);
+		return InitialState == D3D12_RESOURCE_STATE_COMMON;
+	}
+	return false;
+}
+
 struct FD3D12DiscardResource
 {
 	FD3D12DiscardResource(FD3D12Resource* InResource, EResourceTransitionFlags InFlags, uint32 InSubresource = D3D12_RESOURCE_BARRIER_ALL_SUBRESOURCES, const FD3D12Texture* InTexture = nullptr, const FD3D12RenderTargetView* InRTV = nullptr)
@@ -362,7 +392,12 @@ void FD3D12CommandContext::HandleResourceDiscardTransitions(
 {
 	for (const FRHITransitionInfo& Info : TransitionData->TransitionInfos)
 	{
-		if (Info.AccessBefore != ERHIAccess::Discard || Info.Type != FRHITransitionInfo::EType::Texture || !EnumHasAnyFlags(Info.Flags, EResourceTransitionFlags::Clear | EResourceTransitionFlags::Discard))
+		if ((Info.AccessBefore != ERHIAccess::Discard) || !EnumHasAnyFlags(Info.Flags, EResourceTransitionFlags::Clear | EResourceTransitionFlags::Discard))
+		{
+			continue;
+		}
+
+		if (IsImpossibleToDiscardOnAsync(*this, Info))
 		{
 			continue;
 		}
@@ -528,27 +563,32 @@ void FD3D12CommandContext::HandleTransientAliasing(const FD3D12TransitionData* T
 void FD3D12CommandContext::HandleResourceTransitions(const FD3D12TransitionData* TransitionData, bool& bUAVBarrier)
 {
 	const bool bDstAllPipelines = EnumHasAllFlags(TransitionData->DstPipelines, ERHIPipeline::All);
+	const bool bAsyncToAllPipelines = TransitionData->bAsyncToAllPipelines;
+
+	ERHIAccess SplitAccessBefore = ERHIAccess::Unknown;
 
 	for (const FRHITransitionInfo& Info : TransitionData->TransitionInfos)
 	{
 		const bool bUAVAccessAfter = EnumHasAnyFlags(Info.AccessAfter, ERHIAccess::UAVMask);
 
-		// If targeting all pipelines and UAV access then add UAV barrier even with RHI based transitions
-		if (bDstAllPipelines)
-		{
-			bUAVBarrier |= bUAVAccessAfter;
-		}
-
 		// The transition SRV->UAV needs a UAV barrier then to work correctly otherwise there is no synchronization at all
 		bUAVBarrier |= bUAVAccessAfter;
 
+		// Very specific case that needs to be removed with EB
+		// a UAV -> SRVMAsk on the AsyncPipe get split in two: UAV->SRVCompute on Async and SRVCompute->SRVMask on Gfx
+		// On the Async pipe is going to be : UAV->SRVMask(that is automatically converted in UAV->SRVCompute)
+		// On the Direct(Gfx) pipe instead needs to be SRVCompute->SRVMask therefore the check there to change the Before state only on the Direct pipe.
+		SplitAccessBefore = ERHIAccess::Unknown;
+		if ((bAsyncToAllPipelines) && (Info.AccessAfter == ERHIAccess::SRVMask) && (QueueType == ED3D12QueueType::Direct))
+		{
+			SplitAccessBefore = ERHIAccess::SRVCompute;
+		}
 
-
-		const bool bDiscardTransitionOnAsync = (Info.AccessAfter == ERHIAccess::Discard && QueueType == ED3D12QueueType::Async);
+		const bool bImpossibleDiscardTransitionOnAsync = (Info.AccessAfter == ERHIAccess::Discard) && IsImpossibleToDiscardOnAsync(*this, Info);
 
 		// Process transitions which are forced during begin because those contain transition from Graphics to Compute and should
 		// help remove forced patch up command lists for async compute to run on the graphics queue
-		if (Info.Resource && (!bDiscardTransitionOnAsync))
+		if (Info.Resource && (!bImpossibleDiscardTransitionOnAsync))
 		{
 			ProcessResource(*this, Info, [&](const FRHITransitionInfo& Info, FD3D12Resource* Resource, FD3D12Texture* Texture, FD3D12Texture* DiscardTexture)
 				{
@@ -563,8 +603,10 @@ void FD3D12CommandContext::HandleResourceTransitions(const FD3D12TransitionData*
 					// special case that can happen during first frames or resize
 					if ((Info.AccessBefore == ERHIAccess::Present) && (Info.AccessAfter == ERHIAccess::Present))
 						return;
+
+					const ERHIAccess AccessBefore = SplitAccessBefore != ERHIAccess::Unknown ? SplitAccessBefore : Info.AccessBefore;
 						
-					D3D12_RESOURCE_STATES StateBefore = Info.AccessBefore == ERHIAccess::Discard ? GetInitialResourceState(Resource->GetDesc(), QueueType) : GetD3D12ResourceState(Info.AccessBefore, QueueType, Texture);
+					D3D12_RESOURCE_STATES StateBefore = AccessBefore == ERHIAccess::Discard ? GetInitialResourceState(Resource->GetDesc(), QueueType) : GetD3D12ResourceState(Info.AccessBefore, QueueType, Texture);
 						
 					if (Info.AccessBefore != ERHIAccess::Present)
 					{
@@ -655,7 +697,16 @@ void FD3D12CommandContext::RHIEndTransitions(TArrayView<const FRHITransition*> T
 		for (const FRHITransition* Transition : Transitions)
 		{
 			const auto* Data = Transition->GetPrivateData<FD3D12TransitionData>();
-			if (Data->bCrossPipeline)
+
+			if (Data->bAsyncToAllPipelines)
+			{
+				const TRHIPipelineArray<FD3D12SyncPointRef>& DeviceSyncPoints = Data->SyncPoints[GetGPUIndex()];
+				if (CurrentPipeline == ERHIPipeline::Graphics)
+				{
+					WaitSyncPoint(DeviceSyncPoints[ERHIPipeline::AsyncCompute]);
+				}
+			}
+			else if (Data->bCrossPipeline)
 			{
 				const TRHIPipelineArray<FD3D12SyncPointRef>& DeviceSyncPoints = Data->SyncPoints[GetGPUIndex()];
 
@@ -682,7 +733,6 @@ void FD3D12CommandContext::RHIEndTransitions(TArrayView<const FRHITransition*> T
 	{
 		const FD3D12TransitionData* Data = Transition->GetPrivateData<FD3D12TransitionData>();
 
-		//if (!ProcessTransitionDuringBegin(Data))
 		if (ShouldProcessTransition(Data, bIsBeginTransition, CurrentPipeline))
 		{
 			HandleTransientAliasing(Data);
@@ -698,7 +748,6 @@ void FD3D12CommandContext::RHIEndTransitions(TArrayView<const FRHITransition*> T
 		const FD3D12TransitionData* Data = Transition->GetPrivateData<FD3D12TransitionData>();
 
 		// Handle transition during EndTransitions?
-		//if (!ProcessTransitionDuringBegin(Data))
 		if (ShouldProcessTransition(Data, bIsBeginTransition, CurrentPipeline))
 		{
 			HandleResourceTransitions(Data, bUAVBarrier);
@@ -709,6 +758,22 @@ void FD3D12CommandContext::RHIEndTransitions(TArrayView<const FRHITransition*> T
 	{
 		StateCache.FlushComputeShaderCache(true);
 	}
+
+	// Signal fences
+	for (const FRHITransition* Transition : Transitions)
+	{
+		const auto* Data = Transition->GetPrivateData<FD3D12TransitionData>();
+
+		if ((Data->bAsyncToAllPipelines) && (CurrentPipeline == ERHIPipeline::AsyncCompute))
+		{
+			const TRHIPipelineArray<FD3D12SyncPointRef>& DeviceSyncPoints = Data->SyncPoints[GetGPUIndex()];
+
+			if (DeviceSyncPoints[CurrentPipeline])
+			{
+				SignalSyncPoint(DeviceSyncPoints[CurrentPipeline]);
+			}
+		}
+	}
 }
 
 void FD3D12CommandContext::RHISetStaticUniformBuffers(const FUniformBufferStaticBindings& InUniformBuffers)
diff --git a/Engine/Source/Runtime/D3D12RHI/Private/D3D12TransientResourceAllocator.cpp b/Engine/Source/Runtime/D3D12RHI/Private/D3D12TransientResourceAllocator.cpp
--- a/Engine/Source/Runtime/D3D12RHI/Private/D3D12TransientResourceAllocator.cpp
+++ b/Engine/Source/Runtime/D3D12RHI/Private/D3D12TransientResourceAllocator.cpp
@@ -22,7 +22,6 @@ D3D12_RESOURCE_STATES GetInitialResourceState(const D3D12_RESOURCE_DESC& InDesc,
 	{
 		State = D3D12_RESOURCE_STATE_UNORDERED_ACCESS;
 	}
-	check(State != D3D12_RESOURCE_STATE_COMMON);
 	return State;
 }
 
diff --git a/Engine/Source/Runtime/RHI/Private/RHICommandList.cpp b/Engine/Source/Runtime/RHI/Private/RHICommandList.cpp
--- a/Engine/Source/Runtime/RHI/Private/RHICommandList.cpp
+++ b/Engine/Source/Runtime/RHI/Private/RHICommandList.cpp
@@ -2108,7 +2108,7 @@ void FRHICommandListBase::TransitionInternal(TConstArrayView<FRHITransitionInfo>
 		FMemMark Mark(MemStack);
 		FRHITransition* Transition = new (MemStack.Alloc(FRHITransition::GetTotalAllocationSize(), FRHITransition::GetAlignment())) FRHITransition(Pipeline, Pipeline, 
 			EnumHasAnyFlags(CreateFlags, ERHITransitionCreateFlags::AllowDuringRenderPass) ? ETransitionFlag::AllowDuringRenderPass : ETransitionFlag::None);
-		GDynamicRHI->RHICreateTransition(Transition, FRHITransitionCreateInfo(Pipeline, Pipeline, CreateFlags | ERHITransitionCreateFlags::NoSplit | ERHITransitionCreateFlags::NoFence, MakeArrayView(LocalTransitionInfo)));
+		GDynamicRHI->RHICreateTransition(Transition, FRHITransitionCreateInfo(Pipeline, Pipeline, CreateFlags | ERHITransitionCreateFlags::NoSplit, MakeArrayView(LocalTransitionInfo)));
 
 		GetComputeContext().RHIBeginTransitions(MakeArrayView((const FRHITransition**)&Transition, 1));
 		GetComputeContext().RHIEndTransitions(MakeArrayView((const FRHITransition**)&Transition, 1));
@@ -2119,7 +2119,7 @@ void FRHICommandListBase::TransitionInternal(TConstArrayView<FRHITransitionInfo>
 	}
 	else
 	{
-		FRHITransition* Transition = RHICreateTransition(FRHITransitionCreateInfo(Pipeline, Pipeline, CreateFlags | ERHITransitionCreateFlags::NoSplit | ERHITransitionCreateFlags::NoFence, Infos));
+		FRHITransition* Transition = RHICreateTransition(FRHITransitionCreateInfo(Pipeline, Pipeline, CreateFlags | ERHITransitionCreateFlags::NoSplit, Infos));
 		TrackResourceState(MakeArrayView((const FRHITransition**)&Transition, 1));
 		ALLOC_COMMAND(FRHICommandResourceTransition)(Transition);
 	}
