Enable Async Compute by default in DX12 after barriers fixes and driver workaround

diff --git a/Engine/Source/Runtime/D3D12RHI/Private/D3D12Adapter.cpp b/Engine/Source/Runtime/D3D12RHI/Private/D3D12Adapter.cpp
--- a/Engine/Source/Runtime/D3D12RHI/Private/D3D12Adapter.cpp
+++ b/Engine/Source/Runtime/D3D12RHI/Private/D3D12Adapter.cpp
@@ -1172,10 +1172,6 @@ void FD3D12Adapter::InitializeDevices()
 #endif
 		}
 
-#if PLATFORM_WINDOWS
-		GSupportsEfficientAsyncCompute = false;
-#endif
-
 		if (bWithD3DDebug)
 		{
 			CreateD3DInfoQueue();
