Handle BVH resource transitions in d3d12 rhi by ignoring the transition request: UAV barriers are explicitly added inside d3d12 raytracing where needed and BVH resources don't need explicit state changes

diff --git a/Engine/Source/Runtime/D3D12RHI/Private/D3D12Commands.cpp b/Engine/Source/Runtime/D3D12RHI/Private/D3D12Commands.cpp
--- a/Engine/Source/Runtime/D3D12RHI/Private/D3D12Commands.cpp
+++ b/Engine/Source/Runtime/D3D12RHI/Private/D3D12Commands.cpp
@@ -265,6 +265,11 @@ std::tuple<FD3D12Resource*,FD3D12Texture*> GetResourceAndTexture(FD3D12CommandCo
 		}
 		return { nullptr, nullptr };
 	}
+	case FRHITransitionInfo::EType::BVH:
+	{
+		// Nothing special required for BVH transitions - handled inside d3d12 raytracing directly via UAV barriers and don't need explicit state changes
+		return { nullptr, nullptr };
+	}
 	default:
 		checkNoEntry();
 		return { nullptr, nullptr };
@@ -592,7 +597,7 @@ void FD3D12CommandContext::HandleResourceTransitions(const FD3D12TransitionData*
 		{
 			ProcessResource(*this, Info, [&](const FRHITransitionInfo& Info, FD3D12Resource* Resource, FD3D12Texture* Texture, FD3D12Texture* DiscardTexture)
 				{
-					if (!Resource->RequiresResourceStateTracking())
+					if (Resource == nullptr || !Resource->RequiresResourceStateTracking())
 					{
 						return;
 					}
