Remove constant error logging on the edge case where a resource wants to go D3D12_RESOURCE_STATE_SHADING_RATE_SOURCE but GRHISupportsAttachmentVariableRateShading is off

diff --git a/Engine/Source/Runtime/D3D12RHI/Private/D3D12Util.cpp b/Engine/Source/Runtime/D3D12RHI/Private/D3D12Util.cpp
--- a/Engine/Source/Runtime/D3D12RHI/Private/D3D12Util.cpp
+++ b/Engine/Source/Runtime/D3D12RHI/Private/D3D12Util.cpp
@@ -1731,10 +1731,18 @@ D3D12_RESOURCE_STATES GetD3D12ResourceState(ERHIAccess InRHIAccess, ED3D12QueueT
 			}
 			if (EnumHasAnyFlags(InRHIAccess, ERHIAccess::ShadingRateSource))
 			{
+#if !UE_BUILD_SHIPPING
 				if (GRHISupportsAttachmentVariableRateShading == false)
 				{
-					UE_LOG(LogD3D12RHI, Error, TEXT("Resource state is D3D12_RESOURCE_STATE_SHADING_RATE_SOURCE but RHI does not support VRS."));
+					static bool bLogOnce = true;
+					if (bLogOnce)
+					{
+						UE_LOG(LogD3D12RHI, Warning, TEXT("(%s) Resource state is D3D12_RESOURCE_STATE_SHADING_RATE_SOURCE but RHI does not support VRS."), InRHID3D12Texture == nullptr ? TEXT("Unknown") : *InRHID3D12Texture->GetName().GetPlainNameString());
+					}
+					bLogOnce = false;
 				}
+#endif
+
 #if PLATFORM_SUPPORTS_VARIABLE_RATE_SHADING
 				State |= D3D12_RESOURCE_STATE_SHADING_RATE_SOURCE;
 #endif
