bAsyncToAllPipelines is a speciliazed case of bCrossPipeline that signal the fence in another place for safety, therefore set bCrossPipeline to false or the fence can get signal twice

diff --git a/Engine/Source/Runtime/D3D12RHI/Private/D3D12CommandContext.cpp b/Engine/Source/Runtime/D3D12RHI/Private/D3D12CommandContext.cpp
--- a/Engine/Source/Runtime/D3D12RHI/Private/D3D12CommandContext.cpp
+++ b/Engine/Source/Runtime/D3D12RHI/Private/D3D12CommandContext.cpp
@@ -812,6 +812,7 @@ void FD3D12DynamicRHI::RHICreateTransition(FRHITransition* Transition, const FRH
 			if (EnumHasAnyFlags(TransitionInfo.AccessAfter, ERHIAccess::SRVGraphics))
 			{
 				Data->bAsyncToAllPipelines = true;
+				Data->bCrossPipeline = false;
 				break;
 			}
 		}
