Allow extra transitions inside UpdateTexture2D during render passes, Bink is using UpdateTexture2D  while drawing a pass.

diff --git a/Engine/Source/Runtime/D3D12RHI/Private/D3D12Texture.cpp b/Engine/Source/Runtime/D3D12RHI/Private/D3D12Texture.cpp
--- a/Engine/Source/Runtime/D3D12RHI/Private/D3D12Texture.cpp
+++ b/Engine/Source/Runtime/D3D12RHI/Private/D3D12Texture.cpp
@@ -1944,7 +1944,7 @@ void FD3D12Texture::UpdateTexture2D(FRHICommandListBase& RHICmdList, uint32 MipI
 
 	if (RHICmdList.NeedsExtraTransitions())
 	{
-		RHICmdList.TransitionInternal(FRHITransitionInfo(this, ERHIAccess::Unknown, ERHIAccess::CopyDest, EResourceTransitionFlags::IgnoreAfterState, MipIndex, 0, 0));
+		RHICmdList.TransitionInternal(FRHITransitionInfo(this, ERHIAccess::Unknown, ERHIAccess::CopyDest, EResourceTransitionFlags::IgnoreAfterState, MipIndex, 0, 0), ERHITransitionCreateFlags::AllowDuringRenderPass);
 	}
 
 	RHICmdList.EnqueueLambda([
@@ -1963,7 +1963,7 @@ void FD3D12Texture::UpdateTexture2D(FRHICommandListBase& RHICmdList, uint32 MipI
 
 	if (RHICmdList.NeedsExtraTransitions())
 	{
-		RHICmdList.TransitionInternal(FRHITransitionInfo(this, ERHIAccess::CopyDest, ERHIAccess::Unknown, EResourceTransitionFlags::IgnoreAfterState, MipIndex, 0, 0));
+		RHICmdList.TransitionInternal(FRHITransitionInfo(this, ERHIAccess::CopyDest, ERHIAccess::Unknown, EResourceTransitionFlags::IgnoreAfterState, MipIndex, 0, 0), ERHITransitionCreateFlags::AllowDuringRenderPass);
 	}
 }
 
