Safeguard for the case where the current amount of workerthread is 0 in some test

diff --git a/Engine/Source/Runtime/D3D12RHI/Private/D3D12RayTracing.cpp b/Engine/Source/Runtime/D3D12RHI/Private/D3D12RayTracing.cpp
--- a/Engine/Source/Runtime/D3D12RHI/Private/D3D12RayTracing.cpp
+++ b/Engine/Source/Runtime/D3D12RHI/Private/D3D12RayTracing.cpp
@@ -5913,7 +5913,7 @@ void FD3D12CommandContext::RHISetBindingsOnShaderBindingTable(
 
 	FGraphEventArray TaskList;
 
-	const uint32 NumWorkerThreads = FTaskGraphInterface::Get().GetNumWorkerThreads();
+	const uint32 NumWorkerThreads = FMath::Max<uint32>(1, FTaskGraphInterface::Get().GetNumWorkerThreads());
 	const uint32 MaxTasks = FApp::ShouldUseThreadingForPerformance() 
 		? FMath::Min<uint32>(NumWorkerThreads, FD3D12RayTracingShaderBindingTableInternal::MaxBindingWorkers)
 		: 1;
