Enable persistent SBTs by default on Win64 d3d12 RHI

diff --git a/Engine/Source/Runtime/D3D12RHI/Private/Windows/WindowsD3D12Device.cpp b/Engine/Source/Runtime/D3D12RHI/Private/Windows/WindowsD3D12Device.cpp
--- a/Engine/Source/Runtime/D3D12RHI/Private/Windows/WindowsD3D12Device.cpp
+++ b/Engine/Source/Runtime/D3D12RHI/Private/Windows/WindowsD3D12Device.cpp
@@ -1766,6 +1766,7 @@ void FD3D12DynamicRHI::Init()
 	if (GRHIGlobals.RayTracing.Supported)
 	{
 		GRHIGlobals.RayTracing.SupportsAsyncRayTraceDispatch = true;
+		GRHIGlobals.RayTracing.SupportsPersistentSBTs = true;
 	}
 
 	// Indicate that the RHI needs to use the engine's deferred deletion queue.
@@ -2207,6 +2208,8 @@ void FD3D12Queue::ExecuteCommandLists(TArrayView<ID3D12CommandList*> D3DCommandL
 #endif
 )
 {
+	TRACE_CPUPROFILER_EVENT_SCOPE(FD3D12Queue::ExecuteCommandLists);
+
 #if ENABLE_RESIDENCY_MANAGEMENT
 	check(D3DCommandLists.Num() == ResidencySets.Num());
 
