Add LLM scope to track explicit descriptor heap allocations

diff --git a/Engine/Source/Runtime/D3D12RHI/Private/D3D12ExplicitDescriptorCache.cpp b/Engine/Source/Runtime/D3D12RHI/Private/D3D12ExplicitDescriptorCache.cpp
--- a/Engine/Source/Runtime/D3D12RHI/Private/D3D12ExplicitDescriptorCache.cpp
+++ b/Engine/Source/Runtime/D3D12RHI/Private/D3D12ExplicitDescriptorCache.cpp
@@ -8,6 +8,7 @@
 #include "D3D12Stats.h"
 #include "HAL/ConsoleManager.h"
 #include "HAL/PlatformTime.h"
+#include "HAL/LowLevelMemTracker.h"
 #include "Hash/xxhash.h"
 
 // Whether to compare the full descriptor table on cache lookup or only use FXxHash64 digest.
@@ -80,6 +81,8 @@ void FD3D12ExplicitDescriptorHeapCache::ReleaseHeap(FD3D12ExplicitDescriptorHeap
 
 FD3D12ExplicitDescriptorHeapCache::FEntry FD3D12ExplicitDescriptorHeapCache::AllocateHeap(D3D12_DESCRIPTOR_HEAP_TYPE Type, uint32 NumDescriptors)
 {
+	LLM_SCOPE_BYNAME(TEXT("RHIMisc/ExplicitDescriptorHeap"));
+
 	FScopeLock Lock(&CriticalSection);
 
 	// Align request to enable greater reusue in the cache.
