- Only check for multi frame uniform buffers if both the binding record and the SBT are persistent
- Also output the name of the UB when check triggers

diff --git a/Engine/Source/Runtime/D3D12RHI/Private/D3D12RayTracing.cpp b/Engine/Source/Runtime/D3D12RHI/Private/D3D12RayTracing.cpp
--- a/Engine/Source/Runtime/D3D12RHI/Private/D3D12RayTracing.cpp
+++ b/Engine/Source/Runtime/D3D12RHI/Private/D3D12RayTracing.cpp
@@ -5744,8 +5744,8 @@ static void SetRayTracingHitGroup(
 		checkf(GeometrySegmentIndex < NumGeometrySegments, TEXT("Segment %d is out of range for ray tracing geometry '%s' that contains %d segments"),
 			GeometrySegmentIndex, Geometry->DebugName.IsNone() ? TEXT("UNKNOWN") : *Geometry->DebugName.ToString(), NumGeometrySegments);
 
-		// If persistent then all uniform buffers need to be allocated multiframe
-		if (ShaderTable->Lifetime == ERayTracingShaderBindingTableLifetime::Persistent)
+		// If SBT and binding are persistent then all uniform buffers need to be allocated multiframe
+		if (ShaderTable->Lifetime == ERayTracingShaderBindingTableLifetime::Persistent && BindingType == ERayTracingLocalShaderBindingType::Persistent)
 		{
 			for (uint32 UBIndex = 0; UBIndex < NumUniformBuffers; ++UBIndex)
 			{
@@ -5753,7 +5753,7 @@ static void SetRayTracingHitGroup(
 				if (Resource)
 				{
 					FD3D12UniformBuffer* UB = FD3D12CommandContext::RetrieveObject<FD3D12UniformBuffer>(Resource, 0);
-					check(UB->UniformBufferUsage == UniformBuffer_MultiFrame);
+					checkf(UB->UniformBufferUsage == UniformBuffer_MultiFrame, TEXT("Trying to bind non multiframe uniform buffer to persistent SBT: %s"), *Resource->GetLayout().GetDebugName());
 				}
 			}
 		}
