Added checked CopyTextureRegion & CopyBufferRegion
- Small performance hit (~10us) for improved post mortem command closure failures
- Validates buffer and texture regions against its valid bounds
- Added RHI breadcrumb specialized check/log/ulog helpers

diff --git a/Engine/Source/Runtime/D3D12RHI/Private/D3D12Buffer.cpp b/Engine/Source/Runtime/D3D12RHI/Private/D3D12Buffer.cpp
--- a/Engine/Source/Runtime/D3D12RHI/Private/D3D12Buffer.cpp
+++ b/Engine/Source/Runtime/D3D12RHI/Private/D3D12Buffer.cpp
@@ -45,10 +45,10 @@ void FD3D12Buffer::UploadResourceData(FD3D12CommandContext& CommandContext, D3D1
 
 	CommandContext.FlushResourceBarriers();
 
-	CommandContext.GraphicsCommandList()->CopyBufferRegion(
-		Destination->GetResource(),
+	CommandContext.CopyBufferRegionChecked(
+		Destination->GetResource(), Destination->GetName(),
 		DestinationResourceLocation.GetOffsetFromBaseOfResource(),
-		SourceResourceLocation.GetResource()->GetResource(),
+		SourceResourceLocation.GetResource()->GetResource(), FName(),
 		SourceResourceLocation.GetOffsetFromBaseOfResource(),
 		Size);
 
@@ -131,8 +131,26 @@ void FD3D12Buffer::UploadResourceData(FRHICommandListBase& RHICmdList, FResource
 	InResourceArray->Discard();
 }
 
+void FD3D12CommandContext::CopyBufferRegionChecked(
+	ID3D12Resource* DestResource, const FName& DestName, uint64 DestOffset,
+	ID3D12Resource* SourceResource, const FName& SourceName, uint64 SourceOffset,
+	uint32 ByteCount
+)
+{
+#if ENABLE_COPY_BUFFER_REGION_CHECK
+	RHI_BREADCRUMB_CHECK_SHIPPINGF(this, DestOffset + ByteCount <= DestResource->GetDesc().Width, TEXT("Dest byte range out of bounds for: '%s'"), *DestName.ToString());
+	RHI_BREADCRUMB_CHECK_SHIPPINGF(this, SourceOffset + ByteCount <= SourceResource->GetDesc().Width, TEXT("Source byte range out of bounds for: '%s'"), *SourceName.ToString());
+#endif // ENABLE_COPY_BUFFER_REGION_CHECK
+	
+	// Pass down callchain
+	GraphicsCommandList()->CopyBufferRegion(
+		DestResource, DestOffset,
+		SourceResource, SourceOffset,
+		ByteCount
+	);
+}
 
-FD3D12SyncPointRef FD3D12Buffer::UploadResourceDataViaCopyQueue(FResourceArrayUploadInterface* InResourceArray)
+FD3D12SyncPointRef FD3D12Buffer::UploadResourceDataViaCopyQueue(FD3D12CommandContext& OwningContext, FResourceArrayUploadInterface* InResourceArray)
 {
 	TRACE_CPUPROFILER_EVENT_SCOPE(UploadResourceDataViaCopyQueue);
 
@@ -159,6 +177,9 @@ FD3D12SyncPointRef FD3D12Buffer::UploadResourceDataViaCopyQueue(FResourceArrayUp
 		FD3D12CopyScope CopyScope(Device, ED3D12SyncPointType::GPUOnly);
 		SyncPoint = CopyScope.GetSyncPoint();
 
+		RHI_BREADCRUMB_CHECK_SHIPPING(OwningContext, ResourceLocation.GetOffsetFromBaseOfResource() + BufferSize <= ResourceLocation.GetResource()->GetDesc().Width);
+		RHI_BREADCRUMB_CHECK_SHIPPING(OwningContext, SrcResourceLoc.GetOffsetFromBaseOfResource() + BufferSize <= SrcResourceLoc.GetResource()->GetDesc().Width);
+
 		// Perform actual copy op
 		CopyScope.Context.CopyCommandList()->CopyBufferRegion(
 			ResourceLocation.GetResource()->GetResource(),
@@ -593,10 +614,10 @@ void* FD3D12DynamicRHI::LockBuffer(FRHICommandListBase& RHICmdList, FD3D12Buffer
 				Context.UpdateResidency(StagingBuffer);
 				Context.UpdateResidency(pResource);
 
-				Context.GraphicsCommandList()->CopyBufferRegion(
-					StagingBuffer->GetResource(),
+				Context.CopyBufferRegionChecked(
+					StagingBuffer->GetResource(), StagingBuffer->GetName(),
 					0,
-					pResource->GetResource(),
+					pResource->GetResource(), pResource->GetName(),
 					SubAllocOffset + Offset, Size);
 			});
 			
@@ -668,10 +689,10 @@ void FD3D12DynamicRHI::UnlockBuffer(FRHICommandListBase& RHICmdList, FD3D12Buffe
 				Context.UpdateResidency(DestResource);
 				Context.UpdateResidency(SourceResource);
 
-				Context.GraphicsCommandList()->CopyBufferRegion(
-					DestResource->GetResource(),
+				Context.CopyBufferRegionChecked(
+					DestResource->GetResource(), DestResource->GetName(),
 					DestFullOffset,
-					SourceResource->GetResource(),
+					SourceResource->GetResource(), SourceResource->GetName(),
 					SourceFullOffset,
 					LockedData.LockSize
 				);
@@ -831,7 +852,14 @@ void FD3D12CommandContext::RHICopyBufferRegion(FRHIBuffer* DestBufferRHI, uint64
 	FScopedResourceBarrier ScopeResourceBarrierDst(*this, pDestResource  , &DestBuffer->ResourceLocation, D3D12_RESOURCE_STATE_COPY_DEST  , 0);
 	FlushResourceBarriers();
 
-	GraphicsCommandList()->CopyBufferRegion(pDestResource->GetResource(), DestBuffer->ResourceLocation.GetOffsetFromBaseOfResource() + DstOffset, pSourceResource->GetResource(), SourceBuffer->ResourceLocation.GetOffsetFromBaseOfResource() + SrcOffset, NumBytes);
+	CopyBufferRegionChecked(
+		pDestResource->GetResource(), pDestResource->GetName(),
+		DestBuffer->ResourceLocation.GetOffsetFromBaseOfResource() + DstOffset,
+		pSourceResource->GetResource(), pSourceResource->GetName(),
+		SourceBuffer->ResourceLocation.GetOffsetFromBaseOfResource() + SrcOffset,
+		NumBytes
+	);
+	
 	UpdateResidency(pDestResource);
 	UpdateResidency(pSourceResource);
 
diff --git a/Engine/Source/Runtime/D3D12RHI/Private/D3D12CommandContext.cpp b/Engine/Source/Runtime/D3D12RHI/Private/D3D12CommandContext.cpp
--- a/Engine/Source/Runtime/D3D12RHI/Private/D3D12CommandContext.cpp
+++ b/Engine/Source/Runtime/D3D12RHI/Private/D3D12CommandContext.cpp
@@ -1048,7 +1048,7 @@ void FD3D12CommandContextRedirector::RHITransferResources(TConstArrayView<FTrans
 				CD3DX12_TEXTURE_COPY_LOCATION SrcLocation(SrcTexture->GetResource()->GetResource(), 0);
 				CD3DX12_TEXTURE_COPY_LOCATION DstLocation(DstTexture->GetResource()->GetResource(), 0);
 
-				CopyContext->GraphicsCommandList()->CopyTextureRegion(&DstLocation, Box.left, Box.top, Box.front, &SrcLocation, &Box);
+				CopyContext->CopyTextureRegionChecked(&DstLocation, Box.left, Box.top, Box.front, &SrcLocation, &Box, DstTexture->GetFormat(), DstTexture->GetName());
 			}
 		}
 		else
@@ -1202,7 +1202,7 @@ void FD3D12CommandContextRedirector::RHICrossGPUTransfer(TConstArrayView<FTransf
 				CD3DX12_TEXTURE_COPY_LOCATION SrcLocation(SrcTexture->GetResource()->GetResource(), 0);
 				CD3DX12_TEXTURE_COPY_LOCATION DstLocation(DstTexture->GetResource()->GetResource(), 0);
 
-				SrcContext->GraphicsCommandList()->CopyTextureRegion(&DstLocation, Box.left, Box.top, Box.front, &SrcLocation, &Box);
+				SrcContext->CopyTextureRegionChecked(&DstLocation, Box.left, Box.top, Box.front, &SrcLocation, &Box, DstTexture->GetFormat(), DstTexture->GetName());
 			}
 		}
 		else
diff --git a/Engine/Source/Runtime/D3D12RHI/Private/D3D12CommandContext.h b/Engine/Source/Runtime/D3D12RHI/Private/D3D12CommandContext.h
--- a/Engine/Source/Runtime/D3D12RHI/Private/D3D12CommandContext.h
+++ b/Engine/Source/Runtime/D3D12RHI/Private/D3D12CommandContext.h
@@ -481,6 +481,18 @@ protected:
 	FRHIGPUMask PhysicalGPUMask;
 };
 
+#if WITH_RHI_BREADCRUMBS
+inline FRHIBreadcrumbNode* GetTopBreadcrumbNode(FD3D12CommandContextBase& CommandContext)
+{
+	return CommandContext.GetExecutingCommandList().GetCurrentBreadcrumbRef();
+}
+
+inline FRHIBreadcrumbNode* GetTopBreadcrumbNode(FD3D12CommandContextBase* CommandContext)
+{
+	return CommandContext->GetExecutingCommandList().GetCurrentBreadcrumbRef();
+}
+#endif // WITH_RHI_BREADCRUMBS
+
 // RHI Context type used for graphics and async compute command lists.
 class FD3D12CommandContext : public FD3D12ContextCommon, public FD3D12CommandContextBase, public FD3D12DeviceChild
 {
@@ -800,6 +812,19 @@ public:
 		return ShaderBindinglayout ? *ShaderBindinglayout : Default;
 	}
 
+	FORCENOINLINE void CopyBufferRegionChecked(
+		ID3D12Resource* DestResource, const FName& DestName, uint64 DestOffset,
+		ID3D12Resource* SourceResource, const FName& SourceName, uint64 SourceOffset,
+		uint32 ByteCount
+	);
+	
+	FORCENOINLINE void CopyTextureRegionChecked(
+		const D3D12_TEXTURE_COPY_LOCATION* DestCopyLocation, int DestX, int DestY, int DestZ,
+		const D3D12_TEXTURE_COPY_LOCATION* SourceCopyLocation, const D3D12_BOX* SourceBox,
+		EPixelFormat PixelFormat,
+		const FName& DebugName
+	);
+
 protected:
 
 	FD3D12CommandContext* GetSingleDeviceContext(uint32 InGPUIndex) final override
diff --git a/Engine/Source/Runtime/D3D12RHI/Private/D3D12Commands.cpp b/Engine/Source/Runtime/D3D12RHI/Private/D3D12Commands.cpp
--- a/Engine/Source/Runtime/D3D12RHI/Private/D3D12Commands.cpp
+++ b/Engine/Source/Runtime/D3D12RHI/Private/D3D12Commands.cpp
@@ -863,7 +863,13 @@ void FD3D12CommandContext::RHICopyToStagingBuffer(FRHIBuffer* SourceBufferRHI, F
 		}
 		else
 		{
-			GraphicsCommandList()->CopyBufferRegion(pDestResource->GetResource(), DestOffset, pSourceResource->GetResource(), Offset + SourceOffset, NumBytes);
+			CopyBufferRegionChecked(
+				pDestResource->GetResource(), pDestResource->GetName(),
+				DestOffset,
+				pSourceResource->GetResource(), pSourceResource->GetName(),
+				Offset + SourceOffset,
+				NumBytes
+			);
 		}
 		
 		UpdateResidency(pDestResource);
diff --git a/Engine/Source/Runtime/D3D12RHI/Private/D3D12PoolAllocator.cpp b/Engine/Source/Runtime/D3D12RHI/Private/D3D12PoolAllocator.cpp
--- a/Engine/Source/Runtime/D3D12RHI/Private/D3D12PoolAllocator.cpp
+++ b/Engine/Source/Runtime/D3D12RHI/Private/D3D12PoolAllocator.cpp
@@ -942,10 +942,10 @@ void FD3D12PoolAllocator::FlushPendingCopyOps(FD3D12CommandContext& InCommandCon
 				{
 				case FD3D12VRAMCopyOperation::BufferRegion:
 				{
-					InCommandContext.GraphicsCommandList()->CopyBufferRegion(
-						CopyOperation.DestResource->GetResource(),
+					InCommandContext.CopyBufferRegionChecked(
+						CopyOperation.DestResource->GetResource(), CopyOperation.DestResource->GetName(),
 						CopyOperation.DestOffset,
-						CopyOperation.SourceResource->GetResource(),
+						CopyOperation.SourceResource->GetResource(), CopyOperation.SourceResource->GetName(),
 						CopyOperation.SourceOffset,
 						CopyOperation.Size
 					);
diff --git a/Engine/Source/Runtime/D3D12RHI/Private/D3D12RHICommon.h b/Engine/Source/Runtime/D3D12RHI/Private/D3D12RHICommon.h
--- a/Engine/Source/Runtime/D3D12RHI/Private/D3D12RHICommon.h
+++ b/Engine/Source/Runtime/D3D12RHI/Private/D3D12RHICommon.h
@@ -387,3 +387,12 @@ struct TD3D12ObjectPool : public TLockFreePointerListUnordered<TObjectType, PLAT
 	}
 };
 
+// Set of validation to enable in all builds (including shipping)
+
+#ifndef ENABLE_COPY_BUFFER_REGION_CHECK
+	#define ENABLE_COPY_BUFFER_REGION_CHECK  1
+#endif // ENABLE_COPY_BUFFER_REGION_CHECK
+
+#ifndef ENABLE_COPY_TEXTURE_REGION_CHECK
+	#define ENABLE_COPY_TEXTURE_REGION_CHECK 1
+#endif // ENABLE_COPY_TEXTURE_REGION_CHECK
diff --git a/Engine/Source/Runtime/D3D12RHI/Private/D3D12RayTracing.cpp b/Engine/Source/Runtime/D3D12RHI/Private/D3D12RayTracing.cpp
--- a/Engine/Source/Runtime/D3D12RHI/Private/D3D12RayTracing.cpp
+++ b/Engine/Source/Runtime/D3D12RHI/Private/D3D12RayTracing.cpp
@@ -1650,7 +1650,7 @@ public:
 			D3D12_RESOURCE_STATE_COPY_DEST, bHasInitialData, GPUMask, ResourceAllocator, TEXT("Shader binding table"), ERHIAccess::CopyDest);
 
 		// Use copy queue for uploading the data
-		Context.BatchedSyncPoints.ToWait.Emplace(Buffer->UploadResourceDataViaCopyQueue(&Data));
+		Context.BatchedSyncPoints.ToWait.Emplace(Buffer->UploadResourceDataViaCopyQueue(Context, &Data));
 
 		// Enqueue transition to SRV
 		Context.TransitionResource(
@@ -5508,19 +5508,19 @@ static void DispatchRays(FD3D12CommandContext& CommandContext,
 		FMemory::Memcpy(Data, &DispatchDesc, SBTPartSize);
 
 		// Copy SBT data part to resource
-		CommandContext.GraphicsCommandList()->CopyBufferRegion(
-			DispatchRaysDescBufferResource->GetResource(),
+		CommandContext.CopyBufferRegionChecked(
+			DispatchRaysDescBufferResource->GetResource(), DispatchRaysDescBufferResource->GetName(),
 			BaseRayDescBufferOffset,
-			UploadResourceLocation.GetResource()->GetResource(),
+			UploadResourceLocation.GetResource()->GetResource(), UploadResourceLocation.GetResource()->GetName(),
 			UploadResourceLocation.GetOffsetFromBaseOfResource(),
 			SBTPartSize
 		);
 
 		// Copy GPU computed indirect args to resource
-		CommandContext.GraphicsCommandList()->CopyBufferRegion(
-			DispatchRaysDescBufferResource->GetResource(),
+		CommandContext.CopyBufferRegionChecked(
+			DispatchRaysDescBufferResource->GetResource(), DispatchRaysDescBufferResource->GetName(),
 			BaseRayDescBufferOffset + SBTPartSize,
-			ArgumentBuffer->GetResource()->GetResource(),
+			ArgumentBuffer->GetResource()->GetResource(), ArgumentBuffer->GetResource()->GetName(),
 			ArgumentBuffer->ResourceLocation.GetOffsetFromBaseOfResource() + ArgumentOffset,
 			IndirectDimensionSize
 		);
diff --git a/Engine/Source/Runtime/D3D12RHI/Private/D3D12RenderTarget.cpp b/Engine/Source/Runtime/D3D12RHI/Private/D3D12RenderTarget.cpp
--- a/Engine/Source/Runtime/D3D12RHI/Private/D3D12RenderTarget.cpp
+++ b/Engine/Source/Runtime/D3D12RHI/Private/D3D12RenderTarget.cpp
@@ -479,11 +479,14 @@ TRefCountPtr<FD3D12Resource> FD3D12DynamicRHI::GetStagingTexture(FRHITexture* Te
 			Context.FlushResourceBarriers();
 			// Upload heap doesn't need to transition
 
-			Context.GraphicsCommandList()->CopyTextureRegion(
+			Context.CopyTextureRegionChecked(
 				&DestCopyLocation,
 				0, 0, 0,
 				&SourceCopyLocation,
-				RectPtr);
+				RectPtr,
+				Texture->GetFormat(),
+				Texture->GetName()
+			);
 
 			Context.UpdateResidency(Texture->GetResource());
 		});
@@ -792,11 +795,14 @@ void FD3D12DynamicRHI::ReadSurfaceDataMSAARaw(FRHITexture* TextureRHI, FIntRect
 			// Upload heap doesn't need to transition
 
 			// Copy the resolved sample data to the staging texture.
-			Context.GraphicsCommandList()->CopyTextureRegion(
+			Context.CopyTextureRegionChecked(
 				&DestCopyLocation,
 				0, 0, 0,
 				&SourceCopyLocation,
-				&Rect);
+				&Rect,
+				Texture->GetFormat(),
+				Texture->GetName()
+			);
 
 			Context.UpdateResidency(StagingTexture2D);
 			Context.UpdateResidency(NonMSAATexture2D);
@@ -967,11 +973,14 @@ void FD3D12DynamicRHI::RHIReadSurfaceFloatData(FRHITexture* TextureRHI, FIntRect
 		// Don't need to transition upload heaps
 
 		Context.FlushResourceBarriers();
-		Context.GraphicsCommandList()->CopyTextureRegion(
+		Context.CopyTextureRegionChecked(
 			&DestCopyLocation,
 			0, 0, 0,
 			&SourceCopyLocation,
-			&Rect);
+			&Rect,
+			Texture->GetFormat(),
+			Texture->GetName()
+		);
 
 		Context.UpdateResidency(Texture->GetResource());
 	});
@@ -1078,11 +1087,14 @@ void FD3D12DynamicRHI::RHIRead3DSurfaceFloatData(FRHITexture* TextureRHI, FIntRe
 		// Don't need to transition upload heaps
 
 		Context.FlushResourceBarriers();
-		Context.GraphicsCommandList()->CopyTextureRegion(
+		Context.CopyTextureRegionChecked(
 			&DestCopyLocation,
 			0, 0, 0,
 			&SourceCopyLocation,
-			&Rect);
+			&Rect,
+			Texture->GetFormat(),
+			Texture->GetName()
+		);
 
 		Context.UpdateResidency(Texture->GetResource());
 	});
diff --git a/Engine/Source/Runtime/D3D12RHI/Private/D3D12Resources.h b/Engine/Source/Runtime/D3D12RHI/Private/D3D12Resources.h
--- a/Engine/Source/Runtime/D3D12RHI/Private/D3D12Resources.h
+++ b/Engine/Source/Runtime/D3D12RHI/Private/D3D12Resources.h
@@ -1011,7 +1011,7 @@ public:
 	void UploadResourceData(FRHICommandListBase& InRHICmdList, FRHIGPUMask GPUMask, D3D12_RESOURCE_STATES InDestinationState, const void* SourceData, int32 SourceDataSize);
 	void UploadResourceData(FRHICommandListBase& InRHICmdList, FResourceArrayUploadInterface* InResourceArray, D3D12_RESOURCE_STATES InDestinationState, const TCHAR* AssetName = nullptr, const FName& ClassName = NAME_None, const FName& PackageName = NAME_None);
 
-	FD3D12SyncPointRef UploadResourceDataViaCopyQueue(FResourceArrayUploadInterface* InResourceArray);
+	FD3D12SyncPointRef UploadResourceDataViaCopyQueue(FD3D12CommandContext& OwningContext, FResourceArrayUploadInterface* InResourceArray);
 
 	// FRHIResource overrides
 #if RHI_ENABLE_RESOURCE_INFO
diff --git a/Engine/Source/Runtime/D3D12RHI/Private/D3D12Texture.cpp b/Engine/Source/Runtime/D3D12RHI/Private/D3D12Texture.cpp
--- a/Engine/Source/Runtime/D3D12RHI/Private/D3D12Texture.cpp
+++ b/Engine/Source/Runtime/D3D12RHI/Private/D3D12Texture.cpp
@@ -1104,6 +1104,146 @@ uint32 FD3D12DynamicRHI::RHIComputeMemorySize(FRHITexture* TextureRHI)
 	return Texture->ResourceLocation.GetSize();
 }
 
+struct FSubresourceIndices
+{
+	uint32 MipIndex;
+	uint32 SliceIndex;
+	uint32 PlaneIndex;
+};
+
+static FSubresourceIndices GetSubresourceIndices(const D3D12_RESOURCE_DESC& Desc, uint32 Index)
+{
+	const uint32 MipCount   = Desc.MipLevels;
+	const uint32 SliceCount = Desc.Dimension == D3D12_RESOURCE_DIMENSION_TEXTURE3D ? 1 : Desc.DepthOrArraySize;
+
+	return FSubresourceIndices {
+		.MipIndex = Index % MipCount,
+		.SliceIndex = (Index / MipCount) % SliceCount,
+		.PlaneIndex = Index / (MipCount * SliceCount)
+	};
+}
+
+static D3D12_BOX GetSubresourceCopyBox(const D3D12_RESOURCE_DESC& Desc, const D3D12_TEXTURE_COPY_LOCATION* OtherLocation, const FSubresourceIndices& Indices)
+{
+	// Dimensions may be derived from the counter parts footprint
+	if (OtherLocation->Type == D3D12_TEXTURE_COPY_TYPE_PLACED_FOOTPRINT)
+	{
+		return {
+			.left = 0, .top = 0, .front = 0,
+			.right = OtherLocation->PlacedFootprint.Footprint.Width,
+			.bottom = OtherLocation->PlacedFootprint.Footprint.Height,
+			.back = OtherLocation->PlacedFootprint.Footprint.Depth
+		};
+	}
+
+	// Otherwise just assume the entire (sub)resource
+	return {
+		.left = 0, .top = 0, .front = 0,
+		.right = FMath::Max<uint32>(1u, Desc.Width >> Indices.MipIndex),
+		.bottom = FMath::Max<uint32>(1u, Desc.Height >> Indices.MipIndex),
+		.back = Desc.Dimension == D3D12_RESOURCE_DIMENSION_TEXTURE3D ? FMath::Max<uint32>(1u, Desc.DepthOrArraySize >> Indices.MipIndex) : 1
+	};
+}
+
+static uint64 GetPlacedBufferRequiredSize(FD3D12CommandContext* Context, const D3D12_SUBRESOURCE_FOOTPRINT& Footprint, const FPixelFormatInfo& Format, const FName& DebugName)
+{
+	RHI_BREADCRUMB_CHECK_SHIPPINGF(Context, (Footprint.Width % Format.BlockSizeX) == 0, TEXT("Width not aligned to block size"), *DebugName.ToString());
+	RHI_BREADCRUMB_CHECK_SHIPPINGF(Context, (Footprint.Height % Format.BlockSizeY) == 0, TEXT("Height not aligned to block size"), *DebugName.ToString());
+	
+	const uint32 NumColumns = FMath::DivideAndRoundUp(Footprint.Width, static_cast<uint32>(Format.BlockSizeX));
+	const uint32 NumRows    = FMath::DivideAndRoundUp(Footprint.Height, static_cast<uint32>(Format.BlockSizeY));
+
+	// The last row doesn't need the full row pitch, all that matters is that the texel/block starting address is aligned to it
+	const uint32 SubresourceSizeAligned   = Footprint.RowPitch * (NumRows - 1) * Footprint.Depth;
+	const uint32 SubresourceSizeUnaligned = Format.BlockBytes * NumColumns * Footprint.Depth;
+	
+	return SubresourceSizeAligned + SubresourceSizeUnaligned;
+}
+
+FORCENOINLINE void FD3D12CommandContext::CopyTextureRegionChecked(
+	const D3D12_TEXTURE_COPY_LOCATION* DestCopyLocation, int DestX, int DestY, int DestZ,
+	const D3D12_TEXTURE_COPY_LOCATION* SourceCopyLocation, const D3D12_BOX* SourceBox,
+	EPixelFormat PixelFormat,
+	const FName& DebugName
+)
+{
+#if ENABLE_COPY_TEXTURE_REGION_CHECK
+	const D3D12_RESOURCE_DESC DestDesc   = DestCopyLocation->pResource->GetDesc();
+	const D3D12_RESOURCE_DESC SourceDesc = SourceCopyLocation->pResource->GetDesc();
+
+	const FPixelFormatInfo& PixelFormatInfo = GPixelFormats[PixelFormat];
+
+	switch (SourceCopyLocation->Type)
+	{
+		// Copy from texture
+		case D3D12_TEXTURE_COPY_TYPE_SUBRESOURCE_INDEX:
+		{
+			const FSubresourceIndices Indices = GetSubresourceIndices(SourceDesc, SourceCopyLocation->SubresourceIndex);
+
+			const D3D12_BOX CopyBox = SourceBox ? *SourceBox : GetSubresourceCopyBox(SourceDesc, DestCopyLocation, Indices);
+
+			RHI_BREADCRUMB_CHECK_SHIPPINGF(this, (CopyBox.left % PixelFormatInfo.BlockSizeX) == 0 && (CopyBox.right % PixelFormatInfo.BlockSizeX) == 0, TEXT("Width not aligned to block size"), *DebugName.ToString());
+			RHI_BREADCRUMB_CHECK_SHIPPINGF(this, (CopyBox.top % PixelFormatInfo.BlockSizeY) == 0 && (CopyBox.bottom % PixelFormatInfo.BlockSizeY) == 0, TEXT("Height not aligned to block size"), *DebugName.ToString());
+				
+			const uint32 SubresourceWidth  = AlignArbitrary<uint32>(FMath::Max<uint32>(1u, SourceDesc.Width >> Indices.MipIndex), PixelFormatInfo.BlockSizeX);
+			const uint32 SubresourceHeight = AlignArbitrary<uint32>(FMath::Max<uint32>(1u, SourceDesc.Height >> Indices.MipIndex), PixelFormatInfo.BlockSizeY);
+			const uint32 SubresourceDepth  = AlignArbitrary<uint32>(SourceDesc.Dimension == D3D12_RESOURCE_DIMENSION_TEXTURE3D ? FMath::Max<uint32>(1u, SourceDesc.DepthOrArraySize >> Indices.MipIndex) : 1, PixelFormatInfo.BlockSizeZ);
+
+			RHI_BREADCRUMB_CHECK_SHIPPINGF(this, CopyBox.left <= CopyBox.right && CopyBox.right <= SubresourceWidth, TEXT("Source width out of bounds for: '%s'"), *DebugName.ToString());
+			RHI_BREADCRUMB_CHECK_SHIPPINGF(this, CopyBox.top <= CopyBox.bottom && CopyBox.bottom <= SubresourceHeight, TEXT("Source height out of bounds for: '%s'"), *DebugName.ToString());
+			RHI_BREADCRUMB_CHECK_SHIPPINGF(this, CopyBox.front <= CopyBox.back && CopyBox.back <= SubresourceDepth, TEXT("Source depth out of bounds for: '%s'"), *DebugName.ToString());
+			break;
+		}
+
+		// Copy from buffer
+		case D3D12_TEXTURE_COPY_TYPE_PLACED_FOOTPRINT:
+		{
+			const uint32 RequiredPlacedSize = GetPlacedBufferRequiredSize(this, SourceCopyLocation->PlacedFootprint.Footprint, PixelFormatInfo, DebugName);
+			RHI_BREADCRUMB_CHECK_SHIPPINGF(this, SourceCopyLocation->PlacedFootprint.Offset + RequiredPlacedSize <= SourceDesc.Width, TEXT("Source placed buffer width out of bounds for: '%s'"), *DebugName.ToString());
+			break;
+		}
+	}
+
+	switch (DestCopyLocation->Type)
+	{
+		// Copy to texture
+		case D3D12_TEXTURE_COPY_TYPE_SUBRESOURCE_INDEX:
+		{
+			const FSubresourceIndices Indices = GetSubresourceIndices(DestDesc, DestCopyLocation->SubresourceIndex);
+
+			const D3D12_BOX CopyBox = SourceBox ? *SourceBox : GetSubresourceCopyBox(DestDesc, SourceCopyLocation, Indices);
+				
+			RHI_BREADCRUMB_CHECK_SHIPPINGF(this, (CopyBox.left % PixelFormatInfo.BlockSizeX) == 0 && (CopyBox.right % PixelFormatInfo.BlockSizeX) == 0, TEXT("Width not aligned to block size"), *DebugName.ToString());
+			RHI_BREADCRUMB_CHECK_SHIPPINGF(this, (CopyBox.top % PixelFormatInfo.BlockSizeY) == 0 && (CopyBox.bottom % PixelFormatInfo.BlockSizeY) == 0, TEXT("Height not aligned to block size"), *DebugName.ToString());
+				
+			const uint32 SubresourceWidth  = AlignArbitrary<uint32>(FMath::Max<uint32>(1u, DestDesc.Width >> Indices.MipIndex), PixelFormatInfo.BlockSizeX);
+			const uint32 SubresourceHeight = AlignArbitrary<uint32>(FMath::Max<uint32>(1u, DestDesc.Height >> Indices.MipIndex), PixelFormatInfo.BlockSizeY);
+			const uint32 SubresourceDepth  = AlignArbitrary<uint32>(DestDesc.Dimension == D3D12_RESOURCE_DIMENSION_TEXTURE3D ? FMath::Max<uint32>(1u, DestDesc.DepthOrArraySize >> Indices.MipIndex) : 1, PixelFormatInfo.BlockSizeZ);
+
+			RHI_BREADCRUMB_CHECK_SHIPPINGF(this, DestX + (CopyBox.right - CopyBox.left) <= SubresourceWidth, TEXT("Dest width out of bounds for: '%s'"), *DebugName.ToString());
+			RHI_BREADCRUMB_CHECK_SHIPPINGF(this, DestY + (CopyBox.bottom - CopyBox.top) <= SubresourceHeight, TEXT("Dest height out of bounds for: '%s'"), *DebugName.ToString());
+			RHI_BREADCRUMB_CHECK_SHIPPINGF(this, DestZ + (CopyBox.back - CopyBox.front) <= SubresourceDepth, TEXT("Dest depth out of bounds for: '%s'"), *DebugName.ToString());
+			break;
+		}
+
+		// Copy to buffer
+		case D3D12_TEXTURE_COPY_TYPE_PLACED_FOOTPRINT:
+		{
+			const uint32 RequiredPlacedSize = GetPlacedBufferRequiredSize(this, DestCopyLocation->PlacedFootprint.Footprint, PixelFormatInfo, DebugName);
+			RHI_BREADCRUMB_CHECK_SHIPPINGF(this, DestCopyLocation->PlacedFootprint.Offset + RequiredPlacedSize <= DestDesc.Width, TEXT("Dest placed buffer width out of bounds for: '%s'"), *DebugName.ToString());
+			break;
+		}
+	}
+#endif // ENABLE_COPY_TEXTURE_REGION_CHECK
+
+	// Just pass down callchain
+	GraphicsCommandList()->CopyTextureRegion(
+		DestCopyLocation,
+		DestX, DestY, DestZ,
+		SourceCopyLocation,
+		SourceBox
+	);
+}
 
 FTextureRHIRef FD3D12DynamicRHI::AsyncReallocateTexture2D_RenderThread(FRHICommandListImmediate& RHICmdList, FRHITexture* Texture2DRHI, int32 NewMipCount, int32 NewSizeX, int32 NewSizeY, FThreadSafeCounter* RequestStatus)
 {
@@ -1156,12 +1296,15 @@ FTextureRHIRef FD3D12DynamicRHI::AsyncReallocateTexture2D_RenderThread(FRHIComma
 
 				CD3DX12_TEXTURE_COPY_LOCATION DestCopyLocation(DeviceNewTexture->GetResource()->GetResource(), DstSubresource);
 				CD3DX12_TEXTURE_COPY_LOCATION SourceCopyLocation(DeviceOldTexture->GetResource()->GetResource(), SrcSubresource);
-
-				Context.GraphicsCommandList()->CopyTextureRegion(
+				
+				Context.CopyTextureRegionChecked(
 					&DestCopyLocation,
 					0, 0, 0,
 					&SourceCopyLocation,
-					nullptr);
+					nullptr,
+					DeviceNewTexture->GetFormat(),
+					DeviceNewTexture->GetName()
+				);
 
 				Context.UpdateResidency(DeviceNewTexture->GetResource());
 				Context.UpdateResidency(DeviceOldTexture->GetResource());
@@ -1641,11 +1784,14 @@ FRHILockTextureResult FD3D12Texture::Lock(FRHICommandListBase& RHICmdList, const
 			FD3D12CommandContext& Context = FD3D12CommandContext::Get(ExecutingCmdList, 0);
 
 			Context.FlushResourceBarriers();
-			Context.GraphicsCommandList()->CopyTextureRegion(
+			Context.CopyTextureRegionChecked(
 				&DestCopyLocation,
 				0, 0, 0,
 				&SourceCopyLocation,
-				nullptr);
+				nullptr,
+				GetFormat(),
+				GetName()
+			);
 
 			Context.UpdateResidency(GetResource());
 		});
@@ -1677,11 +1823,14 @@ void FD3D12Texture::UpdateTexture(FD3D12CommandContext& Context, uint32 MipIndex
 	CD3DX12_TEXTURE_COPY_LOCATION DestCopyLocation(GetResource()->GetResource(), MipIndex);
 
 	Context.FlushResourceBarriers();
-	Context.GraphicsCommandList()->CopyTextureRegion(
+	Context.CopyTextureRegionChecked(
 		&DestCopyLocation,
 		DestX, DestY, DestZ,
 		&SourceCopyLocation,
-		nullptr);
+		nullptr,
+		GetFormat(),
+		GetName()
+	);
 
 	Context.UpdateResidency(GetResource());
 	
@@ -1699,11 +1848,14 @@ void FD3D12Texture::CopyTextureRegion(FD3D12CommandContext& Context, uint32 Dest
 	FScopedResourceBarrier ConditionalScopeResourceBarrierSrc(Context, SourceTexture->GetResource(), &SourceTexture->ResourceLocation, D3D12_RESOURCE_STATE_COPY_SOURCE, SourceCopyLocation.SubresourceIndex);
 
 	Context.FlushResourceBarriers();
-	Context.GraphicsCommandList()->CopyTextureRegion(
+	Context.CopyTextureRegionChecked(
 		&DestCopyLocation,
 		DestX, DestY, DestZ,
 		&SourceCopyLocation,
-		&SourceBox);
+		&SourceBox,
+		GetFormat(),
+		GetName()
+	);
 
 	Context.UpdateResidency(SourceTexture->GetResource());
 	Context.UpdateResidency(GetResource());
@@ -1818,7 +1970,7 @@ void FD3D12Texture::InitializeTextureData(FRHICommandListBase& RHICmdList, const
 			{
 				Dst.SubresourceIndex = Subresource;
 				Src.PlacedFootprint = Footprints[Subresource];
-				Context.GraphicsCommandList()->CopyTextureRegion(&Dst, 0, 0, 0, &Src, nullptr);
+				Context.CopyTextureRegionChecked(&Dst, 0, 0, 0, &Src, nullptr, Texture->GetFormat(), Texture->GetName());
 			}
 
 			// Update the resource state after the copy has been done (will take care of updating the residency as well)
@@ -2115,13 +2267,16 @@ public:
 
 				RHI_BREADCRUMB_EVENT(Context, "EndMultiUpdateTexture3D");
 
-				Context.GraphicsCommandList()->CopyTextureRegion(
+				Context.CopyTextureRegionChecked(
 					&DestCopyLocation,
 					UpdateInfo.DstStartX,
 					UpdateInfo.DstStartY,
 					UpdateInfo.DstStartZ,
 					&SourceCopyLocation,
-					nullptr);
+					nullptr,
+					TextureLink.GetFormat(),
+					TextureLink.GetName()
+				);
 
 				Context.UpdateResidency(TextureLink.GetResource());
 				DEBUG_EXECUTE_COMMAND_CONTEXT(Context);
@@ -2346,13 +2501,16 @@ void FD3D12DynamicRHI::EndUpdateTexture3D_Internal(FRHICommandListBase& RHICmdLi
 				Context.FlushResourceBarriers();
 				Context.UpdateResidency(Texture.GetResource());
 
-				Context.GraphicsCommandList()->CopyTextureRegion(
+				Context.CopyTextureRegionChecked(
 					&DestCopyLocation,
 					DstStartX,
 					DstStartY,
 					DstStartZ,
 					&SourceCopyLocation,
-					nullptr);
+					nullptr,
+					Texture.GetFormat(),
+					Texture.GetName()
+				);
 
 				Context.ConditionalSplitCommandList();
 				DEBUG_EXECUTE_COMMAND_CONTEXT(Context);
@@ -2674,11 +2832,13 @@ void FD3D12CommandContext::RHICopyTexture(FRHITexture* SourceTextureRHI, FRHITex
 					GetReadBackHeapDescImpl(Dst.PlacedFootprint, GetParentDevice()->GetDevice(), DstDesc, Dst.SubresourceIndex);
 				}
 
-				GraphicsCommandList()->CopyTextureRegion(
+				CopyTextureRegionChecked(
 					&Dst,
 					DestX, DestY, DestZ,
 					&Src,
-					&SrcBox
+					&SrcBox,
+					bReadback ? SourceTexture->GetFormat() : DestTexture->GetFormat(),
+					SourceTexture->GetName()
 				);
 			}
 		}
diff --git a/Engine/Source/Runtime/D3D12RHI/Private/D3D12WorkGraph.cpp b/Engine/Source/Runtime/D3D12RHI/Private/D3D12WorkGraph.cpp
--- a/Engine/Source/Runtime/D3D12RHI/Private/D3D12WorkGraph.cpp
+++ b/Engine/Source/Runtime/D3D12RHI/Private/D3D12WorkGraph.cpp
@@ -958,7 +958,7 @@ void FD3D12CommandContext::DispatchWorkGraphShaderBundle(
 			Desc, 16, FRHIBufferDesc(Desc.Width, 0, BUF_Static), ED3D12ResourceStateMode::MultiState,
 			D3D12_RESOURCE_STATE_COPY_DEST, true, FRHIGPUMask::FromIndex(GetParentDevice()->GetGPUIndex()), nullptr, TEXT("BundleRecordBuffer"), ERHIAccess::CopyDest);
 
-		BatchedSyncPoints.ToWait.Emplace(RootArgBuffer->UploadResourceDataViaCopyQueue(&LocalRootArgs));
+		BatchedSyncPoints.ToWait.Emplace(RootArgBuffer->UploadResourceDataViaCopyQueue(*this, &LocalRootArgs));
 		TransitionResource(RootArgBuffer->GetResource(), D3D12_RESOURCE_STATE_COPY_DEST, D3D12_RESOURCE_STATE_COMMON, 0);
 
 		NodeLocalRootArgumentsTable.StartAddress = RootArgBuffer->ResourceLocation.GetGPUVirtualAddress();
@@ -1303,7 +1303,7 @@ void FD3D12CommandContext::DispatchWorkGraphShaderBundle(
 			Desc, 16, FRHIBufferDesc(Desc.Width, 0, BUF_Static), ED3D12ResourceStateMode::MultiState,
 			D3D12_RESOURCE_STATE_COPY_DEST, true, FRHIGPUMask::FromIndex(GetParentDevice()->GetGPUIndex()), nullptr, TEXT("BundleRecordBuffer"), ERHIAccess::CopyDest);
 
-		BatchedSyncPoints.ToWait.Emplace(RootArgBuffer->UploadResourceDataViaCopyQueue(&LocalRootArgs));
+		BatchedSyncPoints.ToWait.Emplace(RootArgBuffer->UploadResourceDataViaCopyQueue(*this, &LocalRootArgs));
 		TransitionResource(RootArgBuffer->GetResource(), D3D12_RESOURCE_STATE_COPY_DEST, D3D12_RESOURCE_STATE_COMMON, 0);
 
 		NodeLocalRootArgumentsTable.StartAddress = RootArgBuffer->ResourceLocation.GetGPUVirtualAddress();
diff --git a/Engine/Source/Runtime/RHI/Public/RHIBreadcrumbs.h b/Engine/Source/Runtime/RHI/Public/RHIBreadcrumbs.h
--- a/Engine/Source/Runtime/RHI/Public/RHIBreadcrumbs.h
+++ b/Engine/Source/Runtime/RHI/Public/RHIBreadcrumbs.h
@@ -964,6 +964,48 @@ namespace UE::RHI::GPUProfiler
 	#define RHI_BREADCRUMB_EVENT_STR_DEPRECATED(            RHICmdList_Or_RHIContext,            Format, ...) FRHIBreadcrumbEventScope ANONYMOUS_VARIABLE(BreadcrumbEvent)(RHICmdList_Or_RHIContext, FRHIBreadcrumbData(__FILE__, __LINE__), true     , Format, ## __VA_ARGS__)
 	#define RHI_BREADCRUMB_EVENT_CONDITIONAL_STR_DEPRECATED(RHICmdList_Or_RHIContext, Condition, Format, ...) FRHIBreadcrumbEventScope ANONYMOUS_VARIABLE(BreadcrumbEvent)(RHICmdList_Or_RHIContext, FRHIBreadcrumbData(__FILE__, __LINE__), Condition, Format, ## __VA_ARGS__)
 
+	static FString GetSafeBreadcrumbPath(FRHIBreadcrumbNode* Node)
+	{
+		return Node ? Node->GetFullPath() : TEXT("NoBreadcrumb");
+	}
+
+#if DO_CHECK
+	#define RHI_BREADCRUMB_CHECKF(RHICmdList, Condition, Format, ...) \
+	{ \
+		if(UNLIKELY(!(Condition))) \
+		{ \
+			if (FDebug::CheckVerifyFailedImpl2(#Condition, __FILE__, __LINE__, Format TEXT("\nBreadcrumbs: %s"), ##__VA_ARGS__, *GetSafeBreadcrumbPath(GetTopBreadcrumbNode(RHICmdList)))) \
+			{ \
+				PLATFORM_BREAK(); \
+			} \
+			CA_ASSUME(false); \
+		} \
+	}
+#else // DO_CHECK
+	#define RHI_BREADCRUMB_CHECKF(RHICmdList, Condition, Format, ...)
+#endif // DO_CHECK
+
+	#define RHI_BREADCRUMB_CHECK(RHICmdList, Condition) RHI_BREADCRUMB_CHECKF(RHICmdList, Condition, TEXT(""))
+
+	#define RHI_BREADCRUMB_CLOG(RHICmdList, Condition, Log, Severity, Format, ...) \
+	{ \
+		if(UNLIKELY(Condition)) \
+		{ \
+			UE_LOG(Log, Severity, Format TEXT("\nBreadcrumbs: %s"), ##__VA_ARGS__, *GetSafeBreadcrumbPath(GetTopBreadcrumbNode(RHICmdList))); \
+		} \
+	}
+
+	#define RHI_BREADCRUMB_LOG(RHICmdList, Log, Severity, Format, ...) RHI_BREADCRUMB_CLOG(RHICmdList, true, Log, Severity, Format, ##__VA_ARGS__)
+
+#if DO_CHECK
+	#define RHI_BREADCRUMB_CHECK_SHIPPINGF(RHICmdList, Condition, Format, ...) \
+		RHI_BREADCRUMB_CHECKF(RHICmdList, Condition, Format, ##__VA_ARGS__)
+#else // !DO_CHECK
+	#define RHI_BREADCRUMB_CHECK_SHIPPINGF(RHICmdList, Condition, Format, ...) \
+		RHI_BREADCRUMB_CLOG(RHICmdList, !(Condition), LogRHI, Error, TEXT("Check '") TEXT(#Condition) TEXT("' failed. ") Format, ##__VA_ARGS__)
+#endif // !DO_CHECK
+
+	#define RHI_BREADCRUMB_CHECK_SHIPPING(RHICmdList, Condition) RHI_BREADCRUMB_CHECK_SHIPPINGF(RHICmdList, Condition, TEXT(""))
 #else // WITH_RHI_BREADCRUMBS == 0
 
 	#define RHI_BREADCRUMB_EVENT(...)
@@ -974,6 +1016,15 @@ namespace UE::RHI::GPUProfiler
 	#define RHI_BREADCRUMB_EVENT_STR_DEPRECATED(...)
 	#define RHI_BREADCRUMB_EVENT_CONDITIONAL_STR_DEPRECATED(...)
 
+	#define RHI_BREADCRUMB_CHECK(...)
+	#define RHI_BREADCRUMB_CHECKF(...)
+
+	#define RHI_BREADCRUMB_CHECK_LOG(...)
+	#define RHI_BREADCRUMB_CHECK_CLOG(...)
+
+	#define RHI_BREADCRUMB_CHECK_SHIPPINGF(...)
+	#define RHI_BREADCRUMB_CHECK_SHIPPING(...)
+
 #endif
 
 //
diff --git a/Engine/Source/Runtime/RHI/Public/RHICommandList.h b/Engine/Source/Runtime/RHI/Public/RHICommandList.h
--- a/Engine/Source/Runtime/RHI/Public/RHICommandList.h
+++ b/Engine/Source/Runtime/RHI/Public/RHICommandList.h
@@ -1542,6 +1542,16 @@ inline bool IRHIComputeContext::ShouldEmitBreadcrumbs() const
 {
 	return GetExecutingCommandList().CPUBreadcrumbState.bEmitBreadcrumbs;
 }
+
+inline FRHIBreadcrumbNode* GetTopBreadcrumbNode(FRHICommandListBase& RHICmdList)
+{
+	return RHICmdList.GetCurrentBreadcrumbRef();
+}
+
+inline FRHIBreadcrumbNode* GetTopBreadcrumbNode(FRHICommandListBase* RHICmdList)
+{
+	return RHICmdList->GetCurrentBreadcrumbRef();
+}
 #endif
 
 struct FUnnamedRhiCommand
@@ -4366,6 +4376,10 @@ public:
 		// No copy/DMA operation inside render passes
 		check(IsOutsideRenderPass());
 
+		RHI_BREADCRUMB_CHECK_SHIPPING(this, SourceBuffer != DestBuffer);
+		RHI_BREADCRUMB_CHECK_SHIPPING(this, DstOffset + NumBytes <= DestBuffer->GetSize());
+		RHI_BREADCRUMB_CHECK_SHIPPING(this, SrcOffset + NumBytes <= SourceBuffer->GetSize());
+
 		if (Bypass())
 		{
 			GetContext().RHICopyBufferRegion(DestBuffer, DstOffset, SourceBuffer, SrcOffset, NumBytes);
