Fixed checked texture region copy logging, added missing names

diff --git a/Engine/Source/Runtime/D3D12RHI/Private/D3D12Texture.cpp b/Engine/Source/Runtime/D3D12RHI/Private/D3D12Texture.cpp
--- a/Engine/Source/Runtime/D3D12RHI/Private/D3D12Texture.cpp
+++ b/Engine/Source/Runtime/D3D12RHI/Private/D3D12Texture.cpp
@@ -1147,8 +1147,8 @@ static D3D12_BOX GetSubresourceCopyBox(const D3D12_RESOURCE_DESC& Desc, const D3
 
 static uint64 GetPlacedBufferRequiredSize(FD3D12CommandContext* Context, const D3D12_SUBRESOURCE_FOOTPRINT& Footprint, const FPixelFormatInfo& Format, const FName& DebugName)
 {
-	RHI_BREADCRUMB_CHECK_SHIPPINGF(Context, (Footprint.Width % Format.BlockSizeX) == 0, TEXT("Width not aligned to block size"), *DebugName.ToString());
-	RHI_BREADCRUMB_CHECK_SHIPPINGF(Context, (Footprint.Height % Format.BlockSizeY) == 0, TEXT("Height not aligned to block size"), *DebugName.ToString());
+	RHI_BREADCRUMB_CHECK_SHIPPINGF(Context, (Footprint.Width % Format.BlockSizeX) == 0, TEXT("Width not aligned to block size for: '%s'"), *DebugName.ToString());
+	RHI_BREADCRUMB_CHECK_SHIPPINGF(Context, (Footprint.Height % Format.BlockSizeY) == 0, TEXT("Height not aligned to block size for: '%s'"), *DebugName.ToString());
 	
 	const uint32 NumColumns = FMath::DivideAndRoundUp(Footprint.Width, static_cast<uint32>(Format.BlockSizeX));
 	const uint32 NumRows    = FMath::DivideAndRoundUp(Footprint.Height, static_cast<uint32>(Format.BlockSizeY));
@@ -1182,8 +1182,8 @@ FORCENOINLINE void FD3D12CommandContext::CopyTextureRegionChecked(
 
 			const D3D12_BOX CopyBox = SourceBox ? *SourceBox : GetSubresourceCopyBox(SourceDesc, DestCopyLocation, Indices);
 
-			RHI_BREADCRUMB_CHECK_SHIPPINGF(this, (CopyBox.left % PixelFormatInfo.BlockSizeX) == 0 && (CopyBox.right % PixelFormatInfo.BlockSizeX) == 0, TEXT("Width not aligned to block size"), *DebugName.ToString());
-			RHI_BREADCRUMB_CHECK_SHIPPINGF(this, (CopyBox.top % PixelFormatInfo.BlockSizeY) == 0 && (CopyBox.bottom % PixelFormatInfo.BlockSizeY) == 0, TEXT("Height not aligned to block size"), *DebugName.ToString());
+			RHI_BREADCRUMB_CHECK_SHIPPINGF(this, (CopyBox.left % PixelFormatInfo.BlockSizeX) == 0 && (CopyBox.right % PixelFormatInfo.BlockSizeX) == 0, TEXT("Width not aligned to block size for: '%s'"), *DebugName.ToString());
+			RHI_BREADCRUMB_CHECK_SHIPPINGF(this, (CopyBox.top % PixelFormatInfo.BlockSizeY) == 0 && (CopyBox.bottom % PixelFormatInfo.BlockSizeY) == 0, TEXT("Height not aligned to block size for: '%s'"), *DebugName.ToString());
 				
 			const uint32 SubresourceWidth  = AlignArbitrary<uint32>(FMath::Max<uint32>(1u, SourceDesc.Width >> Indices.MipIndex), PixelFormatInfo.BlockSizeX);
 			const uint32 SubresourceHeight = AlignArbitrary<uint32>(FMath::Max<uint32>(1u, SourceDesc.Height >> Indices.MipIndex), PixelFormatInfo.BlockSizeY);
@@ -1213,8 +1213,8 @@ FORCENOINLINE void FD3D12CommandContext::CopyTextureRegionChecked(
 
 			const D3D12_BOX CopyBox = SourceBox ? *SourceBox : GetSubresourceCopyBox(DestDesc, SourceCopyLocation, Indices);
 				
-			RHI_BREADCRUMB_CHECK_SHIPPINGF(this, (CopyBox.left % PixelFormatInfo.BlockSizeX) == 0 && (CopyBox.right % PixelFormatInfo.BlockSizeX) == 0, TEXT("Width not aligned to block size"), *DebugName.ToString());
-			RHI_BREADCRUMB_CHECK_SHIPPINGF(this, (CopyBox.top % PixelFormatInfo.BlockSizeY) == 0 && (CopyBox.bottom % PixelFormatInfo.BlockSizeY) == 0, TEXT("Height not aligned to block size"), *DebugName.ToString());
+			RHI_BREADCRUMB_CHECK_SHIPPINGF(this, (CopyBox.left % PixelFormatInfo.BlockSizeX) == 0 && (CopyBox.right % PixelFormatInfo.BlockSizeX) == 0, TEXT("Width not aligned to block size for: '%s'"), *DebugName.ToString());
+			RHI_BREADCRUMB_CHECK_SHIPPINGF(this, (CopyBox.top % PixelFormatInfo.BlockSizeY) == 0 && (CopyBox.bottom % PixelFormatInfo.BlockSizeY) == 0, TEXT("Height not aligned to block size for: '%s'"), *DebugName.ToString());
 				
 			const uint32 SubresourceWidth  = AlignArbitrary<uint32>(FMath::Max<uint32>(1u, DestDesc.Width >> Indices.MipIndex), PixelFormatInfo.BlockSizeX);
 			const uint32 SubresourceHeight = AlignArbitrary<uint32>(FMath::Max<uint32>(1u, DestDesc.Height >> Indices.MipIndex), PixelFormatInfo.BlockSizeY);
