Fixing D3D12Texture LLM memory tracking.  The system was using the FD3D12Resource* pointer when tracking the texture size, however that pointer was already tracked due to the object itself being in the heap.  This was stomping that object's size to then match the texture size, and eventually we would double-free that pointer from LLM.

#tests Ran with local validation that would assert when the same pointer was used for multiple allocations
[RN] Fixed the LLM tracking of D3D12 texture data on Windows

diff --git a/Engine/Source/Runtime/D3D12RHI/Private/D3D12Texture.cpp b/Engine/Source/Runtime/D3D12RHI/Private/D3D12Texture.cpp
--- a/Engine/Source/Runtime/D3D12RHI/Private/D3D12Texture.cpp
+++ b/Engine/Source/Runtime/D3D12RHI/Private/D3D12Texture.cpp
@@ -105,8 +105,8 @@ void FD3D12TextureStats::UpdateD3D12TextureStats(FD3D12Texture& Texture, const F
 #if PLATFORM_WINDOWS
 		// On Windows there is no way to hook into the low level d3d allocations and frees.
 		// This means that we must manually add the tracking here.
-		LLM_IF_ENABLED(FLowLevelMemTracker::Get().OnLowLevelAlloc(ELLMTracker::Platform, Texture.GetResource(), TextureSize, ELLMTag::GraphicsPlatform));
-		LLM_IF_ENABLED(FLowLevelMemTracker::Get().OnLowLevelAlloc(ELLMTracker::Default, Texture.GetResource(), TextureSize, ELLMTag::Textures));
+		LLM_IF_ENABLED(FLowLevelMemTracker::Get().OnLowLevelAlloc(ELLMTracker::Platform, Texture.ResourceLocation.GetAddressForLLMTracking(), TextureSize, ELLMTag::GraphicsPlatform));
+		LLM_IF_ENABLED(FLowLevelMemTracker::Get().OnLowLevelAlloc(ELLMTracker::Default, Texture.ResourceLocation.GetAddressForLLMTracking(), TextureSize, ELLMTag::Textures));
 		{
 			LLM(UE_MEMSCOPE_DEFAULT(ELLMTag::Textures));
 
@@ -126,8 +126,8 @@ void FD3D12TextureStats::UpdateD3D12TextureStats(FD3D12Texture& Texture, const F
 	else
 	{
 #if PLATFORM_WINDOWS
-		LLM_IF_ENABLED(FLowLevelMemTracker::Get().OnLowLevelFree(ELLMTracker::Platform, Texture.GetResource()));
-		LLM_IF_ENABLED(FLowLevelMemTracker::Get().OnLowLevelFree(ELLMTracker::Default, Texture.GetResource()));
+		LLM_IF_ENABLED(FLowLevelMemTracker::Get().OnLowLevelFree(ELLMTracker::Platform, Texture.ResourceLocation.GetAddressForLLMTracking()));
+		LLM_IF_ENABLED(FLowLevelMemTracker::Get().OnLowLevelFree(ELLMTracker::Default, Texture.ResourceLocation.GetAddressForLLMTracking()));
 
 #if UE_MEMORY_TRACE_ENABLED
 		// Skip back buffers that aren't traced on alloc and don't have valid GPUVirtualAddress
