If an SBT record has a loose parameter resource location allocated then it means it was persistent and now temporarily transient (new RTPSO still compiling for example)
Update the cached loose parameter GPU allocation to make sure the SBT doesn't contain any dangling upload memory allocations when it's persistent again

diff --git a/Engine/Source/Runtime/D3D12RHI/Private/D3D12RayTracing.cpp b/Engine/Source/Runtime/D3D12RHI/Private/D3D12RayTracing.cpp
--- a/Engine/Source/Runtime/D3D12RHI/Private/D3D12RayTracing.cpp
+++ b/Engine/Source/Runtime/D3D12RHI/Private/D3D12RayTracing.cpp
@@ -2207,6 +2207,8 @@ public:
 			FRecordData& RecordData = HitRecordData[RecordIndex];
 			check(RecordData.LooseParameterData.Num() == LooseParameterDataSize);
 			check(FMemory::Memcmp(RecordData.LooseParameterData.GetData(), LooseParameterData, LooseParameterDataSize) == 0);
+			check(RecordData.LooseParameterResourceLocation);
+			OutGPUVirtualAddress = RecordData.LooseParameterResourceLocation->GetGPUVirtualAddress();
 #endif // DO_CHECK
 		}
 		else
@@ -2233,9 +2235,9 @@ public:
 #endif // D3D12RHI_USE_CONSTANT_BUFFER_VIEWS
 
 			// Move ownership of the upload memory to the record data when persistent (otherwise auto freed next 'frame')
+			FRecordData& RecordData = HitRecordData[RecordIndex];
 			if (Lifetime == ERayTracingShaderBindingTableLifetime::Persistent && BindingType == ERayTracingLocalShaderBindingType::Persistent)
 			{
-				FRecordData& RecordData = HitRecordData[RecordIndex];
 				check(RecordData.LooseParameterResourceLocation == nullptr);
 				RecordData.LooseParameterResourceLocation = new FD3D12ResourceLocation(Device);
 				FD3D12ResourceLocation::TransferOwnership(*RecordData.LooseParameterResourceLocation, ResourceLocation);
@@ -2245,6 +2247,16 @@ public:
 				FMemory::Memcpy(RecordData.LooseParameterData.GetData(), LooseParameterData, LooseParameterDataSize);
 #endif // DO_CHECK
 			}
+			// If we had loose parameter resource location allocated then it means it was persistent and now temporarily transient
+			// Update the cached loose parameter GPU allocation to make sure the SBT doesn't contain any dangling upload memory allocations
+			else if (RecordData.LooseParameterResourceLocation)
+			{
+				check(RecordData.LooseParameterData.Num() == LooseParameterDataSize);
+				check(FMemory::Memcmp(RecordData.LooseParameterData.GetData(), LooseParameterData, LooseParameterDataSize) == 0);
+
+				RecordData.LooseParameterResourceLocation->Clear();
+				FD3D12ResourceLocation::TransferOwnership(*RecordData.LooseParameterResourceLocation, ResourceLocation);
+			}
 		}
 
 		return ConstantBufferView;
@@ -4948,14 +4960,6 @@ struct FD3D12RayTracingLocalResourceBinder
 
 	void SetRootCBV(uint32 BaseSlotIndex, uint32 DescriptorIndex, D3D12_GPU_VIRTUAL_ADDRESS Address)
 	{
-#if DO_CHECK
-		// Skip validation for descriptor index 0 when loose parameter parameter data is available because that's
-		// always stored at index 0 and validated at loose parameter set time
-		if (DescriptorIndex == LooseParameterCBVIndex && BindingType == ERayTracingLocalShaderBindingType::Validation && bHasLooseParameterData)
-		{
-			return;
-		}
-#endif // DO_CHECK
 		SetRootDescriptor(BaseSlotIndex, DescriptorIndex, Address);
 	}
 
