Build fix for check using paramaters only defined with #if DO_CHECK

diff --git a/Engine/Source/Runtime/D3D12RHI/Private/D3D12RayTracing.cpp b/Engine/Source/Runtime/D3D12RHI/Private/D3D12RayTracing.cpp
--- a/Engine/Source/Runtime/D3D12RHI/Private/D3D12RayTracing.cpp
+++ b/Engine/Source/Runtime/D3D12RHI/Private/D3D12RayTracing.cpp
@@ -2251,8 +2251,10 @@ public:
 			// Update the cached loose parameter GPU allocation to make sure the SBT doesn't contain any dangling upload memory allocations
 			else if (RecordData.LooseParameterResourceLocation)
 			{
+#if DO_CHECK
 				check(RecordData.LooseParameterData.Num() == LooseParameterDataSize);
 				check(FMemory::Memcmp(RecordData.LooseParameterData.GetData(), LooseParameterData, LooseParameterDataSize) == 0);
+#endif // DO_CHECK
 
 				RecordData.LooseParameterResourceLocation->Clear();
 				FD3D12ResourceLocation::TransferOwnership(*RecordData.LooseParameterResourceLocation, ResourceLocation);
