Adding FatalError when trying to allocate too many descriptors.

diff --git a/Engine/Source/Runtime/D3D12RHI/Private/D3D12BindlessDescriptors.cpp b/Engine/Source/Runtime/D3D12RHI/Private/D3D12BindlessDescriptors.cpp
--- a/Engine/Source/Runtime/D3D12RHI/Private/D3D12BindlessDescriptors.cpp
+++ b/Engine/Source/Runtime/D3D12RHI/Private/D3D12BindlessDescriptors.cpp
@@ -121,6 +121,11 @@ void FD3D12BindlessDescriptorAllocator::Init()
 	BindlessResourcesConfiguration = RHIGetRuntimeBindlessResourcesConfiguration(GMaxRHIShaderPlatform);
 	BindlessSamplersConfiguration = RHIGetRuntimeBindlessSamplersConfiguration(GMaxRHIShaderPlatform);
 
+	MaxResourceHeapSize = GetParentAdapter()->GetMaxDescriptorsForHeapType(ERHIDescriptorHeapType::Standard);
+	MaxSamplerHeapSize = GetParentAdapter()->GetMaxDescriptorsForHeapType(ERHIDescriptorHeapType::Sampler);
+
+	check(MaxResourceHeapSize != 0 && MaxSamplerHeapSize != 0);
+
 	if (BindlessResourcesConfiguration != ERHIBindlessConfiguration::Disabled)
 	{
 		const TStatId Stats[] =
@@ -192,7 +197,13 @@ FRHIDescriptorHandle FD3D12BindlessDescriptorAllocator::AllocateResourceHandle()
 
 		// Grow the descriptor handle allocator
 		uint32 CurrentNumDescriptors = ResourceAllocator->GetCapacity();
-		uint32 NewNumDescriptors = CurrentNumDescriptors * 2;
+		uint32 NewNumDescriptors = FMath::Min<uint32>(CurrentNumDescriptors * 2, MaxResourceHeapSize);
+
+		if (CurrentNumDescriptors == NewNumDescriptors)
+		{
+			UE_LOG(LogD3D12RHI, Fatal, TEXT("Hit D3D12 device limits on descriptors when attempting to allocate a larger descriptor heap."));
+		}
+
 		Result = ResourceAllocator->ResizeGrowAndAllocate(NewNumDescriptors, ResourceAllocator->GetType());
 
 		// Grow the CPU heaps for all devices
diff --git a/Engine/Source/Runtime/D3D12RHI/Private/D3D12BindlessDescriptors.h b/Engine/Source/Runtime/D3D12RHI/Private/D3D12BindlessDescriptors.h
--- a/Engine/Source/Runtime/D3D12RHI/Private/D3D12BindlessDescriptors.h
+++ b/Engine/Source/Runtime/D3D12RHI/Private/D3D12BindlessDescriptors.h
@@ -63,6 +63,9 @@ private:
 	ERHIBindlessConfiguration BindlessResourcesConfiguration{};
 	ERHIBindlessConfiguration BindlessSamplersConfiguration{};
 
+	uint32 MaxResourceHeapSize = 0;
+	uint32 MaxSamplerHeapSize = 0;
+
 	FCriticalSection ResourceHeapsCS;
 	FRHIHeapDescriptorAllocator* ResourceAllocator = nullptr;
 	FRHIHeapDescriptorAllocator* SamplerAllocator = nullptr;
