D3D12 Texture Creation changes
- Shuffling code around to make the Texture Creation refactor easier to integrate.
- Moved CreateD3D12Texture code that didn't require a command list into a new function: CreateTextureInternal
- Moved CreateD3D12Texture code that did require a command list into a new function: FinalizeCreateTexture
- Both functions can be overridden by platforms.
- Removing NeedInitialStateReconcile as CreateTextureInternal can predict the necessary transitions and communicates that in its result.

diff --git a/Engine/Source/Runtime/D3D12RHI/Private/D3D12RHIPrivate.h b/Engine/Source/Runtime/D3D12RHI/Private/D3D12RHIPrivate.h
--- a/Engine/Source/Runtime/D3D12RHI/Private/D3D12RHIPrivate.h
+++ b/Engine/Source/Runtime/D3D12RHI/Private/D3D12RHIPrivate.h
@@ -562,7 +562,43 @@ public:
 
 	virtual FD3D12ResourceDesc GetResourceDesc(const FRHITextureDesc& CreateInfo) const;
 
-	virtual FD3D12Texture* CreateD3D12Texture(const FRHITextureCreateDesc& CreateDesc, FRHICommandListBase* RHICmdList, ID3D12ResourceAllocator* ResourceAllocator = nullptr);
+	// Fills in the ClearValue argument based on the ResourceDesc and CreateDesc flags.
+	D3D12_CLEAR_VALUE* FillClearValue(D3D12_CLEAR_VALUE& ClearValue, const FD3D12ResourceDesc& ResourceDesc, const FRHITextureCreateDesc& CreateDesc) const;
+
+	struct FCreateTextureInternalResult
+	{
+		FD3D12Texture* Texture{};
+		D3D12_RESOURCE_STATES CreateState{};
+		D3D12_RESOURCE_STATES DesiredState{};
+	};
+
+	// Per-platform D3D12Texture creation.
+	virtual FCreateTextureInternalResult CreateTextureInternal(const FRHITextureCreateDesc& CreateDesc, bool bHasInitialData, ID3D12ResourceAllocator* ResourceAllocator = nullptr);
+
+	// Per-platform D3D12Texture post-creation adjustments, like changing the resource state or uploading data.
+	virtual FD3D12Texture* FinalizeCreateTexture(FRHICommandListBase& RHICmdList, const FRHITextureCreateDesc& CreateDesc, const FCreateTextureInternalResult& CreateResult) const;
+
+	// Create D3D12 texture with a resource allocator. Does not currently support any post-creation state or uploading of data.
+	FD3D12Texture* CreateD3D12Texture(const FRHITextureCreateDesc& CreateDesc, ID3D12ResourceAllocator* ResourceAllocator)
+	{
+		check(CreateDesc.BulkData == nullptr);
+		constexpr bool bHasInitialData = false;
+
+		FCreateTextureInternalResult CreateResult = CreateTextureInternal(CreateDesc, bHasInitialData, ResourceAllocator);
+		checkSlow(CreateResult.CreateState == CreateResult.DesiredState);
+		return CreateResult.Texture;
+	}
+
+	// Create D3D12 texture with a command list, meaning we can upload data and adjust resource states
+	FD3D12Texture* CreateD3D12Texture(FRHICommandListBase& RHICmdList, const FRHITextureCreateDesc& CreateDesc)
+	{
+		const bool bHasInitialData = (CreateDesc.BulkData != nullptr);
+		constexpr ID3D12ResourceAllocator* ResourceAllocator = nullptr;
+
+		const FCreateTextureInternalResult CreateResult = CreateTextureInternal(CreateDesc, bHasInitialData, ResourceAllocator);
+		return FinalizeCreateTexture(RHICmdList, CreateDesc, CreateResult);
+	}
+
 	FD3D12Buffer* CreateD3D12Buffer(class FRHICommandListBase* RHICmdList, FRHIBufferDesc const& BufferDesc, ERHIAccess ResourceState, FRHIResourceCreateInfo& CreateInfo, ID3D12ResourceAllocator* ResourceAllocator = nullptr, bool bForceKeepUnlocked = false);
 	virtual FD3D12Texture* CreateNewD3D12Texture(const FRHITextureCreateDesc& CreateDesc, class FD3D12Device* Device);
 
diff --git a/Engine/Source/Runtime/D3D12RHI/Private/D3D12Resources.cpp b/Engine/Source/Runtime/D3D12RHI/Private/D3D12Resources.cpp
--- a/Engine/Source/Runtime/D3D12RHI/Private/D3D12Resources.cpp
+++ b/Engine/Source/Runtime/D3D12RHI/Private/D3D12Resources.cpp
@@ -98,7 +98,6 @@ FD3D12Resource::FD3D12Resource(FD3D12Device* ParentDevice,
 	, bRequiresResidencyTracking(bool(ENABLE_RESIDENCY_MANAGEMENT))
 	, bDepthStencil(false)
 	, bDeferDelete(true)
-	, bReconcileInitialState(false)
 {
 #if UE_BUILD_DEBUG
 	FPlatformAtomics::InterlockedIncrement(&TotalResourceCount);
@@ -772,7 +771,6 @@ HRESULT FD3D12Adapter::CreateCommittedResource(const FD3D12ResourceDesc& InDesc,
 #endif
 
 	HRESULT hr = S_OK;
-	bool bReconcileInitialState = false;
 #if INTEL_EXTENSIONS
 	if (InDesc.bRequires64BitAtomicSupport && IsRHIDeviceIntel() && GDX12INTCAtomicUInt64Emulation)
 	{
@@ -792,7 +790,9 @@ HRESULT FD3D12Adapter::CreateCommittedResource(const FD3D12ResourceDesc& InDesc,
 
 		// Common layout is the required starting state for any "legacy" transitions
 		const D3D12_BARRIER_LAYOUT InitialLayout = D3D12_BARRIER_LAYOUT_COMMON;
-		
+
+		checkf(InInitialState == D3D12_RESOURCE_STATE_COMMON, TEXT("RESOURCE_STATE_COMMON is required for castable resources (Given: %d)"), InInitialState);
+
 		ID3D12ProtectedResourceSession* ProtectedSession = nullptr;
 
 		const TArray<DXGI_FORMAT, TInlineAllocator<4>> CastableFormats = InDesc.GetCastableFormats();
@@ -807,13 +807,6 @@ HRESULT FD3D12Adapter::CreateCommittedResource(const FD3D12ResourceDesc& InDesc,
 			CastableFormats.Num(),
 			CastableFormats.GetData(),
 			IID_PPV_ARGS(pResource.GetInitReference()));
-
-		if (SUCCEEDED(hr))
-		{
-			// Change the initial state to match the initial layout on construction
-			InInitialState = D3D12_RESOURCE_STATE_COMMON;
-			bReconcileInitialState = true;
-		}
 	}
 	else
 #endif
@@ -830,7 +823,6 @@ HRESULT FD3D12Adapter::CreateCommittedResource(const FD3D12ResourceDesc& InDesc,
 		SetName(*ppOutResource, Name);
 
 		(*ppOutResource)->StartTrackingForResidency();
-		(*ppOutResource)->SetInitalStateReconcile(bReconcileInitialState);
 		
 
 		TraceMemoryAllocation(*ppOutResource);
@@ -893,6 +885,8 @@ HRESULT FD3D12Adapter::CreateReservedResource(const FD3D12ResourceDesc& InDesc,
 #if D3D12RHI_SUPPORTS_UNCOMPRESSED_UAV
 	if (InDesc.SupportsUncompressedUAV())
 	{
+		checkf(InInitialState == D3D12_RESOURCE_STATE_COMMON, TEXT("RESOURCE_STATE_COMMON is required for castable resources (Given: %d)"), InInitialState);
+
 		// Common layout is the require starting state for any "legacy" transitions
 		const D3D12_BARRIER_LAYOUT InitialLayout = D3D12_BARRIER_LAYOUT_COMMON;
 
@@ -907,12 +901,6 @@ HRESULT FD3D12Adapter::CreateReservedResource(const FD3D12ResourceDesc& InDesc,
 			CastableFormats.Num(),
 			CastableFormats.GetData(),
 			IID_PPV_ARGS(pResource.GetInitReference()));
-
-		if (SUCCEEDED(hr))
-		{
-			// Change the initial state to match the initial layout on construction
-			InInitialState = D3D12_RESOURCE_STATE_COMMON;
-		}
 	}
 	else
 #endif
@@ -956,7 +944,6 @@ HRESULT FD3D12Adapter::CreatePlacedResource(const FD3D12ResourceDesc& InDesc, FD
 	ID3D12Heap* Heap = BackingHeap->GetHeap();
 
 	TRefCountPtr<ID3D12Resource> pResource;
-	bool bReconcileInitialState = false;
 	HRESULT hr = S_OK;
 #if INTEL_EXTENSIONS
 	if (InDesc.bRequires64BitAtomicSupport && IsRHIDeviceIntel() && GDX12INTCAtomicUInt64Emulation)
@@ -973,6 +960,8 @@ HRESULT FD3D12Adapter::CreatePlacedResource(const FD3D12ResourceDesc& InDesc, FD
 #if D3D12RHI_SUPPORTS_UNCOMPRESSED_UAV
 	if (InDesc.SupportsUncompressedUAV())
 	{
+		checkf(InInitialState == D3D12_RESOURCE_STATE_COMMON, TEXT("RESOURCE_STATE_COMMON is required for castable resources (Given: %d)"), InInitialState);
+
 		// Convert the desc to the version required by CreatePlacedResource2
 		const CD3DX12_RESOURCE_DESC1 LocalDesc(InDesc);
 
@@ -990,13 +979,6 @@ HRESULT FD3D12Adapter::CreatePlacedResource(const FD3D12ResourceDesc& InDesc, FD
 			CastableFormats.Num(),
 			CastableFormats.GetData(),
 			IID_PPV_ARGS(pResource.GetInitReference()));
-
-		if (SUCCEEDED(hr))
-		{
-			// Change the initial state to match the initial layout on construction
-			InInitialState = D3D12_RESOURCE_STATE_COMMON;
-			bReconcileInitialState = true;
-		}
 	}
 	else
 #endif
@@ -1018,8 +1000,6 @@ HRESULT FD3D12Adapter::CreatePlacedResource(const FD3D12ResourceDesc& InDesc, FD
 			InDesc,
 			BackingHeap,
 			HeapDesc.Properties.Type);
-		
-		(*ppOutResource)->SetInitalStateReconcile(bReconcileInitialState);
 
 #if PLATFORM_WINDOWS
 		if (IsTrackingAllAllocations() && BackingHeap->GetHeapDesc().Properties.Type == D3D12_HEAP_TYPE_DEFAULT)
diff --git a/Engine/Source/Runtime/D3D12RHI/Private/D3D12Resources.h b/Engine/Source/Runtime/D3D12RHI/Private/D3D12Resources.h
--- a/Engine/Source/Runtime/D3D12RHI/Private/D3D12Resources.h
+++ b/Engine/Source/Runtime/D3D12RHI/Private/D3D12Resources.h
@@ -222,7 +222,6 @@ private:
 	bool bRequiresResidencyTracking : 1;		//2nd bit
 	bool bDepthStencil : 1;						//3rd bit
 	bool bDeferDelete : 1;						//4th bit
-	bool bReconcileInitialState : 1;			//5th bit
 
 #if UE_BUILD_DEBUG
 	static int64 TotalResourceCount;
@@ -342,9 +341,6 @@ public:
 		return DebugName;
 	}
 
-	inline bool NeedInitialStateReconcile() { return bReconcileInitialState; }
-	inline void SetInitalStateReconcile(bool NewValue) { bReconcileInitialState = NewValue; }
-
 	void DoNotDeferDelete()
 	{
 		bDeferDelete = false;
diff --git a/Engine/Source/Runtime/D3D12RHI/Private/D3D12Texture.cpp b/Engine/Source/Runtime/D3D12RHI/Private/D3D12Texture.cpp
--- a/Engine/Source/Runtime/D3D12RHI/Private/D3D12Texture.cpp
+++ b/Engine/Source/Runtime/D3D12RHI/Private/D3D12Texture.cpp
@@ -585,31 +585,18 @@ FD3D12Texture* FD3D12DynamicRHI::CreateNewD3D12Texture(const FRHITextureCreateDe
 	return new FD3D12Texture(CreateDesc, Device);
 }
 
-FD3D12Texture* FD3D12DynamicRHI::CreateD3D12Texture(const FRHITextureCreateDesc& InCreateDesc, class FRHICommandListBase* RHICmdList, ID3D12ResourceAllocator* ResourceAllocator)
+D3D12_CLEAR_VALUE* FD3D12DynamicRHI::FillClearValue(D3D12_CLEAR_VALUE& ClearValue, const FD3D12ResourceDesc& ResourceDesc, const FRHITextureCreateDesc& CreateDesc) const
 {
-#if PLATFORM_WINDOWS
-	TRACE_CPUPROFILER_EVENT_SCOPE(D3D12RHI::CreateD3D12Texture);
-	UE_TRACE_METADATA_SCOPE_ASSET_FNAME(FName(InCreateDesc.DebugName), InCreateDesc.GetTraceClassName(), InCreateDesc.OwnerName);
-
-	// Make local copy of create desc because certain flags might be removed before actual texture creation
-	FRHITextureCreateDesc CreateDesc = InCreateDesc;
-
-	// Virtual textures currently not supported in default D3D12
-	check(!EnumHasAnyFlags(CreateDesc.Flags, TexCreate_Virtual));
-
-	// Get the resource desc
-	FD3D12ResourceDesc ResourceDesc = GetResourceDesc(CreateDesc);
-
-	const DXGI_FORMAT PlatformResourceFormat = UE::DXGIUtilities::GetPlatformTextureResourceFormat((DXGI_FORMAT)GPixelFormats[CreateDesc.Format].PlatformFormat, CreateDesc.Flags);
 	bool bCreateRTV = EnumHasAnyFlags(ResourceDesc.Flags, D3D12_RESOURCE_FLAG_ALLOW_RENDER_TARGET);
 	bool bCreateDSV = EnumHasAnyFlags(ResourceDesc.Flags, D3D12_RESOURCE_FLAG_ALLOW_DEPTH_STENCIL);
 
-	// Setup the clear value	
+	const DXGI_FORMAT PlatformResourceFormat = UE::DXGIUtilities::GetPlatformTextureResourceFormat((DXGI_FORMAT)GPixelFormats[CreateDesc.Format].PlatformFormat, CreateDesc.Flags);
+
 	D3D12_CLEAR_VALUE* ClearValuePtr = nullptr;
-	D3D12_CLEAR_VALUE ClearValue;
 	if (bCreateDSV && CreateDesc.ClearValue.ColorBinding == EClearBinding::EDepthStencilBound)
 	{
 		const DXGI_FORMAT PlatformDepthStencilFormat = UE::DXGIUtilities::FindDepthStencilFormat(PlatformResourceFormat);
+
 		ClearValue = CD3DX12_CLEAR_VALUE(PlatformDepthStencilFormat, CreateDesc.ClearValue.Value.DSValue.Depth, (uint8)CreateDesc.ClearValue.Value.DSValue.Stencil);
 		ClearValuePtr = &ClearValue;
 	}
@@ -617,13 +604,42 @@ FD3D12Texture* FD3D12DynamicRHI::CreateD3D12Texture(const FRHITextureCreateDesc&
 	{
 		const bool bSRGB = EnumHasAnyFlags(CreateDesc.Flags, TexCreate_SRGB);
 		const DXGI_FORMAT PlatformRenderTargetFormat = UE::DXGIUtilities::FindShaderResourceFormat(PlatformResourceFormat, bSRGB);
+
 		ClearValue = CD3DX12_CLEAR_VALUE(PlatformRenderTargetFormat, CreateDesc.ClearValue.Value.Color);
 		ClearValuePtr = &ClearValue;
 	}
 
+	return ClearValuePtr;
+}
+
+FD3D12DynamicRHI::FCreateTextureInternalResult FD3D12DynamicRHI::CreateTextureInternal(const FRHITextureCreateDesc& InCreateDesc, bool bHasInitialData, ID3D12ResourceAllocator* ResourceAllocator)
+{
+#if PLATFORM_WINDOWS
+	TRACE_CPUPROFILER_EVENT_SCOPE(D3D12RHI::CreateTextureInternal);
+	UE_TRACE_METADATA_SCOPE_ASSET_FNAME(FName(InCreateDesc.DebugName), InCreateDesc.GetTraceClassName(), InCreateDesc.OwnerName);
+
+	// Make local copy of create desc because certain flags might be removed before actual texture creation
+	FRHITextureCreateDesc CreateDesc = InCreateDesc;
+
+	// Virtual textures currently not supported in default D3D12
+	check(!EnumHasAnyFlags(CreateDesc.Flags, TexCreate_Virtual));
+
+	// Get the resource desc
+	FD3D12ResourceDesc ResourceDesc = GetResourceDesc(CreateDesc);
+
+	D3D12_CLEAR_VALUE ClearValue;
+	D3D12_CLEAR_VALUE* ClearValuePtr = FillClearValue(ClearValue, ResourceDesc, CreateDesc);
+
 	const FD3D12Resource::FD3D12ResourceTypeHelper Type(ResourceDesc, D3D12_HEAP_TYPE_DEFAULT);
-	const D3D12_RESOURCE_STATES InitialState = Type.GetOptimalInitialState(CreateDesc.InitialState, true);
-	const D3D12_RESOURCE_STATES CreateState = (CreateDesc.BulkData != nullptr) ? D3D12_RESOURCE_STATE_COPY_DEST : InitialState;
+	const D3D12_RESOURCE_STATES DesiredState = Type.GetOptimalInitialState(CreateDesc.InitialState, true);
+
+	const D3D12_RESOURCE_STATES CreateState =
+		// When creating a resource with castable formats, d3d12 uses enhanced barriers behind the scenes which means we have to start it in the "common" state to use "legacy" barriers
+		ResourceDesc.SupportsUncompressedUAV() ? D3D12_RESOURCE_STATE_COMMON
+		// If we have initial data, we want the resource created in COPY_DEST so we can copy the data immediately
+		: bHasInitialData ? D3D12_RESOURCE_STATE_COPY_DEST
+		// Otherwise, create the resource in the caller's specified state
+		: DesiredState;
 
 	FD3D12Adapter* Adapter = &GetAdapter();
 	FD3D12Texture* D3D12TextureOut = Adapter->CreateLinkedObject<FD3D12Texture>(CreateDesc.GPUMask, [&](FD3D12Device* Device, FD3D12Texture* FirstLinkedObject)
@@ -665,10 +681,7 @@ FD3D12Texture* FD3D12DynamicRHI::CreateD3D12Texture(const FRHITextureCreateDesc&
 			if (ResourceDesc.bReservedResource)
 			{
 				FD3D12Resource* Resource = nullptr;
-				VERIFYD3D12CREATETEXTURERESULT(
-					Adapter->CreateReservedResource(ResourceDesc, Device->GetGPUMask(), CreateState,
-						ED3D12ResourceStateMode::MultiState, CreateState, ClearValuePtr, &Resource, CreateDesc.DebugName, false),
-					ResourceDesc, Device->GetDevice());
+				VERIFYD3D12CREATETEXTURERESULT(Adapter->CreateReservedResource(ResourceDesc, Device->GetGPUMask(), CreateState, ED3D12ResourceStateMode::MultiState, CreateState, ClearValuePtr, &Resource, CreateDesc.DebugName, false), ResourceDesc, Device->GetDevice());
 
 				D3D12_RESOURCE_ALLOCATION_INFO AllocInfo = Device->GetResourceAllocationInfo(ResourceDesc);
 
@@ -681,8 +694,7 @@ FD3D12Texture* FD3D12DynamicRHI::CreateD3D12Texture(const FRHITextureCreateDesc&
 			}
 			else
 			{
-				VERIFYD3D12CREATETEXTURERESULT(Device->GetTextureAllocator().AllocateTexture(
-					ResourceDesc, ClearValuePtr, CreateDesc.Format, Location, CreateState, CreateDesc.DebugName), ResourceDesc, Device->GetDevice());
+				VERIFYD3D12CREATETEXTURERESULT(Device->GetTextureAllocator().AllocateTexture(ResourceDesc, ClearValuePtr, CreateDesc.Format, Location, CreateState, CreateDesc.DebugName), ResourceDesc, Device->GetDevice());
 			}
 
 			Location.SetOwner(NewTexture);
@@ -702,35 +714,13 @@ FD3D12Texture* FD3D12DynamicRHI::CreateD3D12Texture(const FRHITextureCreateDesc&
 		}
 
 		// Unlock immediately if no initial data
-		if (CreateDesc.BulkData == nullptr)
+		if (!bHasInitialData)
 		{
 			Location.UnlockPoolData();
 		}
 
 		check(Location.IsValid());
 
-		{
-			FD3D12Resource* CreatedResource = Location.GetResource();
-		
-			if (CreatedResource->NeedInitialStateReconcile())
-			{
-				const D3D12_RESOURCE_STATES RealInitialState = CreatedResource->GetInitialResourceState();
-
-				if (InitialState != RealInitialState)
-				{
-					uint32 GPUIndex = Device->GetGPUIndex();
-					RHICmdList->EnqueueLambda([CreatedResource, RealInitialState, InitialState, GPUIndex](FRHICommandListBase& ExecutingCmdList)
-						{
-#if UE_BUILD_DEBUG
-							UE_LOG(LogD3D12RHI, Warning, TEXT("Resource %s - Initial State Transition %s -> %s"), *CreatedResource->GetName().ToString(), *ConvertToResourceStateString(RealInitialState), *ConvertToResourceStateString(InitialState));
-#endif
-							FD3D12CommandContext& Context = FD3D12CommandContext::Get(ExecutingCmdList, GPUIndex);
-							Context.AddTransitionBarrier(CreatedResource, RealInitialState, InitialState, D3D12_RESOURCE_BARRIER_ALL_SUBRESOURCES);
-						});
-				}
-			}
-		}
-
 		if (ResourceDesc.NeedsUAVAliasWorkarounds())
 		{
 			Adapter->CreateUAVAliasResource(ClearValuePtr, CreateDesc.DebugName, Location);
@@ -750,28 +740,61 @@ FD3D12Texture* FD3D12DynamicRHI::CreateD3D12Texture(const FRHITextureCreateDesc&
 
 	FD3D12TextureStats::D3D12TextureAllocated(*D3D12TextureOut);
 
-	// Initialize if data is given
-	if (CreateDesc.BulkData != nullptr)
-	{
-		check(RHICmdList);
-		D3D12TextureOut->InitializeTextureData(*RHICmdList, CreateDesc, InitialState);
-		CreateDesc.BulkData->Discard();
-	}
+	FCreateTextureInternalResult Result{};
+	Result.Texture = D3D12TextureOut;
+	Result.CreateState = CreateState;
+	Result.DesiredState = DesiredState;
 
 	check( (D3D12TextureOut->GetTrackedSubResourceCount()) == (D3D12TextureOut->ResourceLocation.GetResource()->GetSubresourceCount()) );
 
-	return D3D12TextureOut;
+	return Result;
 #else
 	checkf(false, TEXT("XBOX_CODE_MERGE : Removed. The Xbox platform version should be used."));
-	return nullptr;
+	FCreateTextureInternalResult Result{};
+	return Result;
 #endif // PLATFORM_WINDOWS
 }
 
+static void ReconcileInitialState(FRHICommandListBase& RHICmdList, FD3D12Texture* Texture, D3D12_RESOURCE_STATES CurrentState, D3D12_RESOURCE_STATES NeededState)
+{
+	if (CurrentState != NeededState)
+	{
+		RHICmdList.EnqueueLambda([Texture, CurrentState, NeededState](FRHICommandListBase& ExecutingCmdList)
+		{
+			for (FD3D12Texture::FLinkedObjectIterator TextureIt(Texture); TextureIt; ++TextureIt)
+			{
+				const uint32 GPUIndex = TextureIt->GetParentDevice()->GetGPUIndex();
+
+				FD3D12CommandContext& Context = FD3D12CommandContext::Get(ExecutingCmdList, GPUIndex);
+				Context.AddTransitionBarrier(TextureIt->GetResource(), CurrentState, NeededState, D3D12_RESOURCE_BARRIER_ALL_SUBRESOURCES);
+			}
+		});
+	}
+}
+
+FD3D12Texture* FD3D12DynamicRHI::FinalizeCreateTexture(FRHICommandListBase& RHICmdList, const FRHITextureCreateDesc& CreateDesc, const FCreateTextureInternalResult& CreateResult) const
+{
+	if (CreateDesc.BulkData)
+	{
+		// Make sure the resource is in the copy dest state before enqueuing the copy
+		ReconcileInitialState(RHICmdList, CreateResult.Texture, CreateResult.CreateState, D3D12_RESOURCE_STATE_COPY_DEST);
+
+		// After the copy, make sure the resource is in the requested state
+		CreateResult.Texture->InitializeTextureData(RHICmdList, CreateDesc, CreateResult.DesiredState);
+		CreateDesc.BulkData->Discard();
+	}
+	else
+	{
+		// Make sure the resource is in the requested state
+		ReconcileInitialState(RHICmdList, CreateResult.Texture, CreateResult.CreateState, CreateResult.DesiredState);
+	}
+
+	return CreateResult.Texture;
+}
 
 FTextureRHIRef FD3D12DynamicRHI::RHICreateTexture(FRHICommandListBase& RHICmdList, const FRHITextureCreateDesc& CreateDesc)
 {
-	ID3D12ResourceAllocator* ResourceAllocator = nullptr;
-	return CreateD3D12Texture(CreateDesc, &RHICmdList, ResourceAllocator);
+	return CreateD3D12Texture(RHICmdList, CreateDesc);
 }
 
 class FWaitInitialMipDataUploadTask
@@ -1260,8 +1283,7 @@ FTextureRHIRef FD3D12DynamicRHI::AsyncReallocateTexture2D_RenderThread(FRHIComma
 	);
 
 	// Allocate a new texture.
-	ID3D12ResourceAllocator* ResourceAllocator = nullptr;
-	FD3D12Texture* NewTexture = CreateD3D12Texture(CreateDesc, &RHICmdList, ResourceAllocator);
+	FD3D12Texture* NewTexture = CreateD3D12Texture(RHICmdList, CreateDesc);
 
 	RHICmdList.EnqueueLambda([
 		RootOldTexture = OldTexture,
diff --git a/Engine/Source/Runtime/D3D12RHI/Private/D3D12TransientResourceAllocator.cpp b/Engine/Source/Runtime/D3D12RHI/Private/D3D12TransientResourceAllocator.cpp
--- a/Engine/Source/Runtime/D3D12RHI/Private/D3D12TransientResourceAllocator.cpp
+++ b/Engine/Source/Runtime/D3D12RHI/Private/D3D12TransientResourceAllocator.cpp
@@ -177,7 +177,8 @@ FRHITransientTexture* FD3D12TransientResourceHeapAllocator::CreateTexture(const
 			TRACE_CPUPROFILER_EVENT_SCOPE(AllocatePlacedTexture);
 			FResourceAllocatorAdapter ResourceAllocatorAdapter(Adapter, Heap, Allocation, Desc);
 			FRHITextureCreateDesc CreateDesc(InCreateInfo, ERHIAccess::Discard, InDebugName);
-			FRHITexture* Texture = DynamicRHI->CreateD3D12Texture(CreateDesc, nullptr, &ResourceAllocatorAdapter);
+
+			FD3D12Texture* Texture = DynamicRHI->CreateD3D12Texture(CreateDesc, &ResourceAllocatorAdapter);
 			return FRHITransientResource::FResourceTaskResult{ Texture, ResourceAllocatorAdapter.GpuVirtualAddress };
 
 		}, Initializer.Hash, Info.SizeInBytes, InCreateInfo);
