Only touch HitRecordData when SBT is persistent (otherwise the array is not initialized)

diff --git a/Engine/Source/Runtime/D3D12RHI/Private/D3D12RayTracing.cpp b/Engine/Source/Runtime/D3D12RHI/Private/D3D12RayTracing.cpp
--- a/Engine/Source/Runtime/D3D12RHI/Private/D3D12RayTracing.cpp
+++ b/Engine/Source/Runtime/D3D12RHI/Private/D3D12RayTracing.cpp
@@ -2232,32 +2232,35 @@ public:
 			ConstantBufferView = new FD3D12ConstantBufferView(Device, nullptr);
 			ConstantBufferView->CreateView(&ResourceLocation, 0, Align(LooseParameterDataSize, D3D12_CONSTANT_BUFFER_DATA_PLACEMENT_ALIGNMENT));
 			WorkerData[WorkerIndex].LooseParameterCBVs.Add(ConstantBufferView);			
-#endif // D3D12RHI_USE_CONSTANT_BUFFER_VIEWS
-
-			// Move ownership of the upload memory to the record data when persistent (otherwise auto freed next 'frame')
-			FRecordData& RecordData = HitRecordData[RecordIndex];
-			if (Lifetime == ERayTracingShaderBindingTableLifetime::Persistent && BindingType == ERayTracingLocalShaderBindingType::Persistent)
-			{
-				check(RecordData.LooseParameterResourceLocation == nullptr);
-				RecordData.LooseParameterResourceLocation = new FD3D12ResourceLocation(Device);
-				FD3D12ResourceLocation::TransferOwnership(*RecordData.LooseParameterResourceLocation, ResourceLocation);
+			#endif // D3D12RHI_USE_CONSTANT_BUFFER_VIEWS
 
-#if DO_CHECK
-				RecordData.LooseParameterData.SetNum(LooseParameterDataSize);
-				FMemory::Memcpy(RecordData.LooseParameterData.GetData(), LooseParameterData, LooseParameterDataSize);
-#endif // DO_CHECK
-			}
-			// If we had loose parameter resource location allocated then it means it was persistent and now temporarily transient
-			// Update the cached loose parameter GPU allocation to make sure the SBT doesn't contain any dangling upload memory allocations
-			else if (RecordData.LooseParameterResourceLocation)
+			if (Lifetime == ERayTracingShaderBindingTableLifetime::Persistent)
 			{
-#if DO_CHECK
-				check(RecordData.LooseParameterData.Num() == LooseParameterDataSize);
-				check(FMemory::Memcmp(RecordData.LooseParameterData.GetData(), LooseParameterData, LooseParameterDataSize) == 0);
-#endif // DO_CHECK
+				// Move ownership of the upload memory to the record data when persistent (otherwise auto freed next 'frame')
+				FRecordData& RecordData = HitRecordData[RecordIndex];
+				if (BindingType == ERayTracingLocalShaderBindingType::Persistent)
+				{
+					check(RecordData.LooseParameterResourceLocation == nullptr);
+					RecordData.LooseParameterResourceLocation = new FD3D12ResourceLocation(Device);
+					FD3D12ResourceLocation::TransferOwnership(*RecordData.LooseParameterResourceLocation, ResourceLocation);
+
+				#if DO_CHECK
+					RecordData.LooseParameterData.SetNum(LooseParameterDataSize);
+					FMemory::Memcpy(RecordData.LooseParameterData.GetData(), LooseParameterData, LooseParameterDataSize);
+				#endif // DO_CHECK
+				}
+				// If we had loose parameter resource location allocated then it means it was persistent and now temporarily transient
+				// Update the cached loose parameter GPU allocation to make sure the SBT doesn't contain any dangling upload memory allocations
+				else if (RecordData.LooseParameterResourceLocation)
+				{
+				#if DO_CHECK
+					check(RecordData.LooseParameterData.Num() == LooseParameterDataSize);
+					check(FMemory::Memcmp(RecordData.LooseParameterData.GetData(), LooseParameterData, LooseParameterDataSize) == 0);
+				#endif // DO_CHECK
 
-				RecordData.LooseParameterResourceLocation->Clear();
-				FD3D12ResourceLocation::TransferOwnership(*RecordData.LooseParameterResourceLocation, ResourceLocation);
+					RecordData.LooseParameterResourceLocation->Clear();
+					FD3D12ResourceLocation::TransferOwnership(*RecordData.LooseParameterResourceLocation, ResourceLocation);
+				}
 			}
 		}
 
