PR #12714: BugFix for uint32 overflow for disc cache size in bytes

diff --git a/Engine/Source/Runtime/D3D12RHI/Private/Windows/WindowsD3D12DiskCache.cpp b/Engine/Source/Runtime/D3D12RHI/Private/Windows/WindowsD3D12DiskCache.cpp
--- a/Engine/Source/Runtime/D3D12RHI/Private/Windows/WindowsD3D12DiskCache.cpp
+++ b/Engine/Source/Runtime/D3D12RHI/Private/Windows/WindowsD3D12DiskCache.cpp
@@ -119,7 +119,7 @@ void FDiskCacheInterface::GrowMapping(SIZE_T size, bool firstrun)
 		mCurrentFileMapSize = fileSize;
 	}
 
-	mMemoryMap = CreateFileMapping(mFile, NULL, PAGE_READWRITE, 0, (uint32)mCurrentFileMapSize, NULL);
+	mMemoryMap = CreateFileMapping(mFile, NULL, PAGE_READWRITE, (uint32)(mCurrentFileMapSize >> 32), (uint32)(mCurrentFileMapSize & 0xFFFFFFFF), NULL);
 	if (mMemoryMap == (HANDLE)nullptr)
 	{
 		//error state!
diff --git a/Engine/Source/Runtime/D3D12RHI/Private/Windows/WindowsD3D12DiskCache.h b/Engine/Source/Runtime/D3D12RHI/Private/Windows/WindowsD3D12DiskCache.h
--- a/Engine/Source/Runtime/D3D12RHI/Private/Windows/WindowsD3D12DiskCache.h
+++ b/Engine/Source/Runtime/D3D12RHI/Private/Windows/WindowsD3D12DiskCache.h
@@ -19,7 +19,7 @@ class FDiskCacheInterface
 	{
 		uint32 mHeaderVersion;
 		uint32 mNumPsos;
-		uint32 mSizeInBytes; // The number of bytes after the header
+		size_t mSizeInBytes; // The number of bytes after the header
 	};
 
 private:
