Update the D3D12 SamplerState cache to use a hybrid LRU cache, replacing the least recently used unreferenced SamplerState when hitting the cache limit rather than crashing.
This extends the global unique sampler limit from 2048 to 65536, though we're still limited to 2048 samplers per frame.
Remove D3D12.SamplerWarningThreshold/GD3D12SamplerWarningThreshold, since we can now warn exactly when we reach the limit.

diff --git a/Engine/Source/Runtime/D3D12RHI/Private/D3D12Device.cpp b/Engine/Source/Runtime/D3D12RHI/Private/D3D12Device.cpp
--- a/Engine/Source/Runtime/D3D12RHI/Private/D3D12Device.cpp
+++ b/Engine/Source/Runtime/D3D12RHI/Private/D3D12Device.cpp
@@ -110,6 +110,7 @@ FD3D12Device::FD3D12Device(FRHIGPUMask InGPUMask, FD3D12Adapter* InAdapter)
 #endif
 	, GlobalSamplerHeap        (this)
 	, OnlineDescriptorManager  (this)
+	, SamplerCache			(D3D12_MAX_SHADER_VISIBLE_SAMPLER_HEAP_SIZE)
 	, DefaultBufferAllocator   (this, FRHIGPUMask::All()) //Note: Cross node buffers are possible 
 	, DefaultFastAllocator     (this, FRHIGPUMask::All(), D3D12_HEAP_TYPE_UPLOAD, 1024 * 1024 * 4)
 	, TextureAllocator         (this, FRHIGPUMask::All())
@@ -156,7 +157,7 @@ FD3D12Device::~FD3D12Device()
 	TextureAllocator.CleanUpAllocations();
 	TextureAllocator.Destroy();
 
-	SamplerMap.Empty();
+	SamplerCache.Empty();
 }
 
 FD3D12Device::FResidencyManager::FResidencyManager(FD3D12Device& Parent)
diff --git a/Engine/Source/Runtime/D3D12RHI/Private/D3D12Device.h b/Engine/Source/Runtime/D3D12RHI/Private/D3D12Device.h
--- a/Engine/Source/Runtime/D3D12RHI/Private/D3D12Device.h
+++ b/Engine/Source/Runtime/D3D12RHI/Private/D3D12Device.h
@@ -18,6 +18,8 @@ D3D12Device.h: D3D12 Device Interfaces
 #include "D3D12Submission.h"
 #include "D3D12GPUProfiler.h"
 
+#include "Containers/LruCache.h"
+
 class FD3D12Device;
 class FD3D12DynamicRHI;
 class FD3D12Buffer;
@@ -336,7 +338,7 @@ private:
 
 	TArray<FD3D12Queue, TInlineAllocator<(uint32)ED3D12QueueType::Count>> Queues;
 
-	TMap<D3D12_SAMPLER_DESC, TRefCountPtr<FD3D12SamplerState>> SamplerMap;
+	TLruCache<D3D12_SAMPLER_DESC, TRefCountPtr<FD3D12SamplerState>> SamplerCache;
 	uint32 SamplerID = 0;
 
 	/** Hashmap used to cache resource allocation size information */
diff --git a/Engine/Source/Runtime/D3D12RHI/Private/D3D12State.cpp b/Engine/Source/Runtime/D3D12RHI/Private/D3D12State.cpp
--- a/Engine/Source/Runtime/D3D12RHI/Private/D3D12State.cpp
+++ b/Engine/Source/Runtime/D3D12RHI/Private/D3D12State.cpp
@@ -225,17 +225,11 @@ FSamplerStateRHIRef FD3D12DynamicRHI::RHICreateSamplerState(const FSamplerStateI
 	});
 }
 
-int32 GD3D12SamplerWarningThreshold = 10;
-static FAutoConsoleVariableRef CVarD3D12SamplerWarningThreshold(
-	TEXT("D3D12.SamplerWarningThreshold"),
-	GD3D12SamplerWarningThreshold,
-	TEXT("Threshold to start warning about creating too many sampler states")
-);
-
 static void LogSamplerStateWarning(const FSamplerStateInitializerRHI& Initializer)
 {
 	UE_LOG(LogD3D12RHI, Warning,
-		TEXT("Approaching SamplerState limit: FSamplerStateInitializerRHI(Filter: %d, AddressU: %d, AddressV: %d, AddressW: %d, MipBias: %f, MinMipLevel: %f, MaxMipLevel: %f, MaxAnisotropy: %d, BorderColor: %d, SamplerComparisonFunction: %d)"),
+		TEXT("New SamplerState would exceed cache limit: FSamplerStateInitializerRHI(Filter: %d, AddressU: %d, AddressV: %d, AddressW: %d, MipBias: %f, MinMipLevel: %f, MaxMipLevel: %f, MaxAnisotropy: %d, BorderColor: %d, SamplerComparisonFunction: %d).")
+		TEXT("An unreferenced SamplerState will be replaced. Try reducing r.ViewTextureMipBias.Quantization and reloading your project to reduce the number of unique samplers when Dynamic Resoluton is active."),
 		Initializer.Filter.GetIntValue(),
 		Initializer.AddressU.GetIntValue(),
 		Initializer.AddressV.GetIntValue(),
@@ -301,7 +295,7 @@ FD3D12SamplerState* FD3D12Device::CreateSampler(const FSamplerStateInitializerRH
 
 	// Check to see if the sampler has already been created
 	// This is done to reduce cache misses accessing sampler objects
-	TRefCountPtr<FD3D12SamplerState>* PreviouslyCreated = SamplerMap.Find(SamplerDesc);
+	const TRefCountPtr<FD3D12SamplerState>* PreviouslyCreated = SamplerCache.FindAndTouch(SamplerDesc);
 	if (PreviouslyCreated)
 	{
 		return PreviouslyCreated->GetReference();
@@ -311,14 +305,36 @@ FD3D12SamplerState* FD3D12Device::CreateSampler(const FSamplerStateInitializerRH
 		// 16-bit IDs are used for faster hashing
 		check(SamplerID < 0xffff);
 
-		if (static_cast<int32>(SamplerID + 1) > (D3D12_MAX_SHADER_VISIBLE_SAMPLER_HEAP_SIZE - GD3D12SamplerWarningThreshold))
+		// If we're full and the least recent entry is still referenced, touch all referenced keys starting from the least recent until we find an unreferenced one
+		// Samplers are generally added in the following order:
+		//    1. Perennial samplers that will be used in many places and are never recreated (2-100+ references)
+		//    2. Older material samplers that are no longer being used (1 reference, from the map itself)
+		//    3. Recently created material samplers that are in use (2 references)
+		// This process will move perennial samplers that were created before any extant material samplers to the top of the recently used list and ensure we have a sampler in the LRU slot that's safe to remove
+		if (SamplerCache.Num() >= D3D12_MAX_SHADER_VISIBLE_SAMPLER_HEAP_SIZE)
 		{
 			LogSamplerStateWarning(Initializer);
+
+			int32 EntriesTouched = 0;
+			D3D12_SAMPLER_DESC CurrentKey = SamplerCache.GetLeastRecentKey();
+
+			while (SamplerCache.Find(CurrentKey)->GetReference()->GetRefCount() > 1)
+			{
+				static_cast<void>(SamplerCache.FindAndTouch(CurrentKey)); // We don't need the return value, just to mark it as recently used
+				CurrentKey = SamplerCache.GetLeastRecentKey();
+				
+				EntriesTouched++;
+				checkf(EntriesTouched < D3D12_MAX_SHADER_VISIBLE_SAMPLER_HEAP_SIZE, TEXT("Attempted to create new SamplerState when cache is full, and no unreferenced entries are available for replacement."));
+			}
+			
+			// Once we have the least recently used unreferenced entry, remove it
+			// We need to free the descriptor (and possibly the bindless handle) immediately so we don't overshoot the limit while waiting for it to be deleted
+			SamplerCache.Find(CurrentKey)->GetReference()->FreeDescriptor();
+			SamplerCache.Remove(CurrentKey);
 		}
 
 		FD3D12SamplerState* NewSampler = new FD3D12SamplerState(this, SamplerDesc, static_cast<uint16>(SamplerID), FirstLinkedObject);
-
-		SamplerMap.Add(SamplerDesc, NewSampler);
+		SamplerCache.Add(SamplerDesc, NewSampler);
 
 		SamplerID++;
 
@@ -801,7 +817,7 @@ FD3D12SamplerState::FD3D12SamplerState(FD3D12Device* InParent, const D3D12_SAMPL
 #endif
 }
 
-FD3D12SamplerState::~FD3D12SamplerState()
+void FD3D12SamplerState::FreeDescriptor()
 {
 	if (OfflineDescriptor)
 	{
@@ -812,8 +828,13 @@ FD3D12SamplerState::~FD3D12SamplerState()
 		// Handle is shared -- freeing is handled by the head link in the FD3D12LinkedAdapterObject
 		if (BindlessHandle.IsValid() && IsHeadLink())
 		{
-			GetParentDevice()->GetBindlessDescriptorManager().DeferredFreeFromDestructor(BindlessHandle);
+			GetParentDevice()->GetBindlessDescriptorManager().ImmediateFree(BindlessHandle);
 		}
 #endif
 	}
 }
+
+FD3D12SamplerState::~FD3D12SamplerState()
+{
+	FreeDescriptor();
+}
diff --git a/Engine/Source/Runtime/D3D12RHI/Private/D3D12State.h b/Engine/Source/Runtime/D3D12RHI/Private/D3D12State.h
--- a/Engine/Source/Runtime/D3D12RHI/Private/D3D12State.h
+++ b/Engine/Source/Runtime/D3D12RHI/Private/D3D12State.h
@@ -20,6 +20,8 @@ public:
 	FD3D12SamplerState(FD3D12Device* InParent, const D3D12_SAMPLER_DESC& Desc, uint16 SamplerID, FD3D12SamplerState* FirstLinkedObject);
 	~FD3D12SamplerState();
 
+	void FreeDescriptor();
+
 	virtual FRHIDescriptorHandle GetBindlessHandle() const final { return BindlessHandle; }
 };
 
