Fix crash in D3D12 RHI when playing videos
- We weren't considering the PreExecute callback on payloads when deciding whether to increment the fence counter.
- The correct function to use is RequiresQueueFenceSignal()

diff --git a/Engine/Source/Runtime/D3D12RHI/Private/D3D12Submission.cpp b/Engine/Source/Runtime/D3D12RHI/Private/D3D12Submission.cpp
--- a/Engine/Source/Runtime/D3D12RHI/Private/D3D12Submission.cpp
+++ b/Engine/Source/Runtime/D3D12RHI/Private/D3D12Submission.cpp
@@ -681,7 +681,7 @@ uint64 FD3D12Queue::FinalizePayload(bool bRequiresSignal, FPayloadArray& Payload
 	PayloadToSubmit->CompletionFenceValue = NextCompletionValue;
 	PayloadToSubmit->bAlwaysSignal |= bRequiresSignal;
 
-	if (PayloadToSubmit->bAlwaysSignal || PayloadToSubmit->SyncPointsToSignal.Num())
+	if (PayloadToSubmit->RequiresQueueFenceSignal())
 	{
 		++Fence.NextCompletionValue;
 	}
