Try to initialize Aftermath only on NVIDIA adapters

diff --git a/Engine/Source/Runtime/D3D12RHI/Private/D3D12Adapter.cpp b/Engine/Source/Runtime/D3D12RHI/Private/D3D12Adapter.cpp
--- a/Engine/Source/Runtime/D3D12RHI/Private/D3D12Adapter.cpp
+++ b/Engine/Source/Runtime/D3D12RHI/Private/D3D12Adapter.cpp
@@ -393,7 +393,10 @@ void FD3D12Adapter::CreateRootDevice(bool bWithDebug)
 #if PLATFORM_WINDOWS
 	
 #if NV_AFTERMATH
-	UE::RHICore::Nvidia::Aftermath::InitializeBeforeDeviceCreation();
+	if (IsRHIDeviceNVIDIA())
+	{
+		UE::RHICore::Nvidia::Aftermath::InitializeBeforeDeviceCreation();
+	}
 #endif
 
 	// Setup DRED if requested
@@ -612,7 +615,10 @@ void FD3D12Adapter::CreateRootDevice(bool bWithDebug)
 #endif // ENABLE_RESIDENCY_MANAGEMENT
 
 #if NV_AFTERMATH
-	UE::RHICore::Nvidia::Aftermath::D3D12::InitializeDevice(RootDevice);
+	if (IsRHIDeviceNVIDIA())
+	{
+		UE::RHICore::Nvidia::Aftermath::D3D12::InitializeDevice(RootDevice);
+	}
 #endif
 
 #if PLATFORM_WINDOWS
