Fix occasional crash when coming back from suspend mode

Fix: Utilize the RHI suspend lock around our async texture uploading

diff --git a/Engine/Source/Runtime/D3D12RHI/Private/D3D12Texture.cpp b/Engine/Source/Runtime/D3D12RHI/Private/D3D12Texture.cpp
--- a/Engine/Source/Runtime/D3D12RHI/Private/D3D12Texture.cpp
+++ b/Engine/Source/Runtime/D3D12RHI/Private/D3D12Texture.cpp
@@ -909,6 +909,10 @@ FTextureRHIRef FD3D12DynamicRHI::RHIAsyncCreateTexture2D(uint32 SizeX, uint32 Si
 	// Also remember that ALL touched resources in a COPY command list decay to COMMON after ExecuteCommandLists completes.
 	const D3D12_RESOURCE_STATES InitialState = D3D12_RESOURCE_STATE_COMMON;
 
+#if !PLATFORM_WINDOWS
+	FRWScopeLock ReadLock(GetID3D12XboxCommonDynamicRHI()->GetRHISuspendedLock(), SLT_ReadOnly);
+#endif
+
 	FD3D12Adapter* Adapter = &GetAdapter();
 	FD3D12Texture* TextureOut = Adapter->CreateLinkedObject<FD3D12Texture>(FRHIGPUMask::All(), [&](FD3D12Device* Device, FD3D12Texture* FirstLinkedObject)
 	{
