Merging RHI CreateBuffer refactor to Main.

diff --git a/Engine/Plugins/Runtime/GeometryCache/Source/GeometryCache/Private/GeometryCacheSceneProxy.cpp b/Engine/Plugins/Runtime/GeometryCache/Source/GeometryCache/Private/GeometryCacheSceneProxy.cpp
--- a/Engine/Plugins/Runtime/GeometryCache/Source/GeometryCache/Private/GeometryCacheSceneProxy.cpp
+++ b/Engine/Plugins/Runtime/GeometryCache/Source/GeometryCache/Private/GeometryCacheSceneProxy.cpp
@@ -1490,10 +1490,19 @@ void FGeomCacheVertexFactory::Init(FRHICommandListBase& RHICmdList, const FVerte
 	SetData(RHICmdList, NewData);
 }
 
+static FBufferRHIRef CreateGeomCacheIndexBuffer(FRHICommandListBase& RHICmdList, int32 NumAllocatedIndices)
+{
+	const FRHIBufferCreateDesc CreateDesc =
+		FRHIBufferCreateDesc::CreateIndex<uint32>(TEXT("FGeomCacheIndexBuffer"), NumAllocatedIndices)
+		.AddUsage(EBufferUsageFlags::Dynamic | EBufferUsageFlags::ShaderResource)
+		.SetInitialState(ERHIAccess::VertexOrIndexBuffer | ERHIAccess::SRVMask);
+
+	return RHICmdList.CreateBuffer(CreateDesc);
+}
+
 void FGeomCacheIndexBuffer::InitRHI(FRHICommandListBase& RHICmdList)
 {
-	FRHIResourceCreateInfo CreateInfo(TEXT("FGeomCacheIndexBuffer"));
-	IndexBufferRHI = RHICmdList.CreateBuffer(NumAllocatedIndices * sizeof(uint32), BUF_Dynamic | BUF_IndexBuffer | BUF_ShaderResource, sizeof(uint32), ERHIAccess::VertexOrIndexBuffer | ERHIAccess::SRVMask, CreateInfo);
+	IndexBufferRHI = CreateGeomCacheIndexBuffer(RHICmdList, NumAllocatedIndices);
 	NumValidIndices = 0;
 
 	if (IndexBufferRHI && NumAllocatedIndices)
@@ -1525,8 +1534,7 @@ void FGeomCacheIndexBuffer::Update(FRHICommandListBase& RHICmdList, const TArray
 	if (Indices.Num() > NumAllocatedIndices)
 	{
 		NumAllocatedIndices = Indices.Num();
-		FRHIResourceCreateInfo CreateInfo(TEXT("FGeomCacheIndexBuffer"));
-		IndexBufferRHI = RHICmdList.CreateBuffer(NumAllocatedIndices * sizeof(uint32), BUF_Dynamic | BUF_IndexBuffer | BUF_ShaderResource, sizeof(uint32), ERHIAccess::VertexOrIndexBuffer | ERHIAccess::SRVMask, CreateInfo);
+		IndexBufferRHI = CreateGeomCacheIndexBuffer(RHICmdList, NumAllocatedIndices);
 		bReallocate = true;
 	}
 
@@ -1570,8 +1578,7 @@ void FGeomCacheIndexBuffer::UpdateSizeOnly(FRHICommandListBase& RHICmdList, int3
 	bool bReallocate = false;
 	if (NewNumIndices > NumAllocatedIndices)
 	{
-		FRHIResourceCreateInfo CreateInfo(TEXT("FGeomCacheIndexBuffer"));
-		IndexBufferRHI = RHICmdList.CreateBuffer(NewNumIndices * sizeof(uint32), BUF_Dynamic | BUF_IndexBuffer | BUF_ShaderResource, sizeof(uint32), ERHIAccess::VertexOrIndexBuffer | ERHIAccess::SRVMask, CreateInfo);
+		IndexBufferRHI = CreateGeomCacheIndexBuffer(RHICmdList, NewNumIndices);
 		NumAllocatedIndices = NewNumIndices;
 		NumValidIndices = 0;
 		bReallocate = true;
@@ -1587,10 +1594,19 @@ void FGeomCacheIndexBuffer::UpdateSizeOnly(FRHICommandListBase& RHICmdList, int3
 	}
 }
 
+static FBufferRHIRef CreateGeomCacheVertexBuffer(FRHICommandListBase& RHICmdList, const TCHAR* Name, int32 SizeInBytes)
+{
+	const FRHIBufferCreateDesc CreateDesc =
+		FRHIBufferCreateDesc::CreateVertex(TEXT("FGeomCacheVertexBuffer"), SizeInBytes)
+		.AddUsage(EBufferUsageFlags::Static | EBufferUsageFlags::ShaderResource)
+		.SetInitialState(ERHIAccess::VertexOrIndexBuffer | ERHIAccess::SRVMask);
+
+	return RHICmdList.CreateBuffer(CreateDesc);
+}
+
 void FGeomCacheVertexBuffer::InitRHI(FRHICommandListBase& RHICmdList)
 {
-	FRHIResourceCreateInfo CreateInfo(TEXT("FGeomCacheVertexBuffer"));
-	VertexBufferRHI = RHICmdList.CreateBuffer(SizeInBytes, BUF_Static | BUF_VertexBuffer | BUF_ShaderResource, 0, ERHIAccess::VertexOrIndexBuffer | ERHIAccess::SRVMask, CreateInfo);
+	VertexBufferRHI = CreateGeomCacheVertexBuffer(RHICmdList, TEXT("FGeomCacheVertexBuffer"), SizeInBytes);
 
 	if (VertexBufferRHI && RHISupportsManualVertexFetch(GMaxRHIShaderPlatform))
 	{
@@ -1610,8 +1626,7 @@ void FGeomCacheVertexBuffer::ReleaseRHI()
 
 void FGeomCacheTangentBuffer::InitRHI(FRHICommandListBase& RHICmdList)
 {
-	FRHIResourceCreateInfo CreateInfo(TEXT("FGeomCacheTangentBuffer"));
-	VertexBufferRHI = RHICmdList.CreateBuffer(SizeInBytes, BUF_Static | BUF_VertexBuffer | BUF_ShaderResource, 0, ERHIAccess::VertexOrIndexBuffer | ERHIAccess::SRVMask, CreateInfo);
+	VertexBufferRHI = CreateGeomCacheVertexBuffer(RHICmdList, TEXT("FGeomCacheTangentBuffer"), SizeInBytes);
 
 	if (VertexBufferRHI && RHISupportsManualVertexFetch(GMaxRHIShaderPlatform))
 	{
@@ -1625,8 +1640,7 @@ void FGeomCacheTangentBuffer::InitRHI(FRHICommandListBase& RHICmdList)
 
 void FGeomCacheColorBuffer::InitRHI(FRHICommandListBase& RHICmdList)
 {
-	FRHIResourceCreateInfo CreateInfo(TEXT("FGeomCacheColorBuffer"));
-	VertexBufferRHI = RHICmdList.CreateBuffer(SizeInBytes, BUF_Static | BUF_VertexBuffer | BUF_ShaderResource, 0, ERHIAccess::VertexOrIndexBuffer | ERHIAccess::SRVMask, CreateInfo);
+	VertexBufferRHI = CreateGeomCacheVertexBuffer(RHICmdList, TEXT("FGeomCacheColorBuffer"), SizeInBytes);
 
 	if (VertexBufferRHI && RHISupportsManualVertexFetch(GMaxRHIShaderPlatform))
 	{
@@ -1648,8 +1662,7 @@ void FGeomCacheVertexBuffer::UpdateRaw(FRHICommandListBase& RHICmdList, const vo
 	if (NewSizeInBytes > SizeInBytes)
 	{
 		SizeInBytes = NewSizeInBytes;
-		FRHIResourceCreateInfo CreateInfo(TEXT("FGeomCacheVertexBuffer"));
-		VertexBufferRHI = RHICmdList.CreateBuffer(SizeInBytes, BUF_Static | BUF_VertexBuffer | BUF_ShaderResource, 0, ERHIAccess::VertexOrIndexBuffer | ERHIAccess::SRVMask, CreateInfo);
+		VertexBufferRHI = CreateGeomCacheVertexBuffer(RHICmdList, TEXT("FGeomCacheVertexBuffer"), SizeInBytes);
 		bReallocate = true;
 	}
 
@@ -1689,8 +1702,7 @@ void FGeomCacheVertexBuffer::UpdateSize(FRHICommandListBase& RHICmdList, int32 N
 	if (NewSizeInBytes > SizeInBytes)
 	{
 		SizeInBytes = NewSizeInBytes;
-		FRHIResourceCreateInfo CreateInfo(TEXT("FGeomCacheVertexBuffer"));
-		VertexBufferRHI = RHICmdList.CreateBuffer(SizeInBytes, BUF_Static | BUF_VertexBuffer | BUF_ShaderResource, 0, ERHIAccess::VertexOrIndexBuffer | ERHIAccess::SRVMask, CreateInfo);
+		VertexBufferRHI = CreateGeomCacheVertexBuffer(RHICmdList, TEXT("FGeomCacheVertexBuffer"), SizeInBytes);
 		bReallocate = true;
 	}
 
diff --git a/Engine/Plugins/Runtime/HairStrands/Source/HairStrandsCore/Private/GroomResources.cpp b/Engine/Plugins/Runtime/HairStrands/Source/HairStrandsCore/Private/GroomResources.cpp
--- a/Engine/Plugins/Runtime/HairStrands/Source/HairStrandsCore/Private/GroomResources.cpp
+++ b/Engine/Plugins/Runtime/HairStrands/Source/HairStrandsCore/Private/GroomResources.cpp
@@ -168,17 +168,22 @@ void CreateBuffer(FRHICommandListBase& RHICmdList, const TArray<typename FormatT
 	const uint32 DataCount = InData.Num();
 	const uint32 DataSizeInBytes = FormatType::SizeInByte * DataCount;
 
-	if (DataSizeInBytes == 0) return;
+	if (DataSizeInBytes == 0)
+	{
+		return;
+	}
 
-	FRHIResourceCreateInfo CreateInfo(DebugName);
-	CreateInfo.ResourceArray = nullptr;
+	const FRHIBufferCreateDesc CreateDesc =
+		FRHIBufferCreateDesc::CreateVertex(DebugName, DataSizeInBytes)
+		.AddUsage(EBufferUsageFlags::Static | EBufferUsageFlags::ShaderResource)
+		.SetInitialState(InitialAccess)
+		.SetOwnerName(OwnerName)
+		.SetInitActionInitializer();
 
-	OutBuffer.VertexBufferRHI = RHICmdList.CreateVertexBuffer(DataSizeInBytes, BUF_Static | BUF_ShaderResource, InitialAccess, CreateInfo);
-	OutBuffer.VertexBufferRHI->SetOwnerName(OwnerName);
+	FRHIBufferInitializer Initializer = RHICmdList.CreateBufferInitializer(CreateDesc);
+	Initializer.WriteData(InData.GetData(), DataSizeInBytes);
 
-	void* BufferData = RHICmdList.LockBuffer(OutBuffer.VertexBufferRHI, 0, DataSizeInBytes, RLM_WriteOnly);
-	FMemory::Memcpy(BufferData, InData.GetData(), DataSizeInBytes);
-	RHICmdList.UnlockBuffer(OutBuffer.VertexBufferRHI);
+	OutBuffer.VertexBufferRHI = Initializer.Finalize();
 	OutBuffer.ShaderResourceViewRHI = RHICmdList.CreateShaderResourceView(
 		OutBuffer.VertexBufferRHI, 
 		FRHIViewDesc::CreateBufferSRV()
@@ -991,14 +996,16 @@ FHairCardIndexBuffer::FHairCardIndexBuffer(const TArray<FHairCardsIndexFormat::T
 
 void FHairCardIndexBuffer::InitRHI(FRHICommandListBase& RHICmdList)
 {
-	const uint32 DataSizeInBytes = FHairCardsIndexFormat::SizeInByte * Indices.Num();
+	FResourceArrayUploadArrayView UploadView(MakeConstArrayView(Indices));
+
+	const FRHIBufferCreateDesc CreateDesc =
+		FRHIBufferCreateDesc::CreateIndex<FHairCardsIndexFormat::Type>(TEXT("FHairCardIndexBuffer"), Indices.Num())
+		.AddUsage(EBufferUsageFlags::Static | EBufferUsageFlags::ShaderResource)
+		.SetOwnerName(GetOwnerName())
+		.SetInitActionResourceArray(&UploadView)
+		.SetInitialState(ERHIAccess::VertexOrIndexBuffer | ERHIAccess::SRVCompute);
 
-	FRHIResourceCreateInfo CreateInfo(TEXT("FHairCardIndexBuffer"));
-	IndexBufferRHI = RHICmdList.CreateBuffer(DataSizeInBytes, BUF_Static | BUF_IndexBuffer | BUF_ShaderResource, FHairCardsIndexFormat::SizeInByte, ERHIAccess::VertexOrIndexBuffer | ERHIAccess::SRVCompute, CreateInfo);
-	void* Buffer = RHICmdList.LockBuffer(IndexBufferRHI, 0, DataSizeInBytes, RLM_WriteOnly);
-	FMemory::Memcpy(Buffer, Indices.GetData(), DataSizeInBytes);
-	RHICmdList.UnlockBuffer(IndexBufferRHI);
-	IndexBufferRHI->SetOwnerName(GetOwnerName());
+	IndexBufferRHI = RHICmdList.CreateBuffer(CreateDesc);
 }
 
 FHairCardsRestResource::FHairCardsRestResource(const FHairCardsBulkData& InBulkData, const FHairResourceName& InResourceName, const FName& InOwnerName) :
diff --git a/Engine/Plugins/Runtime/HairStrands/Source/HairStrandsCore/Private/HairStrandsVertexFactory.cpp b/Engine/Plugins/Runtime/HairStrands/Source/HairStrandsCore/Private/HairStrandsVertexFactory.cpp
--- a/Engine/Plugins/Runtime/HairStrands/Source/HairStrandsCore/Private/HairStrandsVertexFactory.cpp
+++ b/Engine/Plugins/Runtime/HairStrands/Source/HairStrandsCore/Private/HairStrandsVertexFactory.cpp
@@ -45,23 +45,25 @@ public:
 	{
 		const static FLazyName ClassName(TEXT("FDummyCulledDispatchVertexIdsBuffer"));
 		{
-			FRHIResourceCreateInfo CreateInfo(TEXT("FDummyCulledDispatchVertexIdsBuffer"));
-			CreateInfo.ClassName = ClassName;
-			uint32 NumBytes = sizeof(uint32) * 4;
-			VertexBufferRHI = RHICmdList.CreateBuffer(NumBytes, BUF_Static | BUF_VertexBuffer | BUF_ShaderResource, 0, ERHIAccess::VertexOrIndexBuffer | ERHIAccess::SRVMask, CreateInfo);
-			uint32* DummyContents = (uint32*)RHICmdList.LockBuffer(VertexBufferRHI, 0, NumBytes, RLM_WriteOnly);
-			DummyContents[0] = DummyContents[1] = DummyContents[2] = DummyContents[3] = 0;
-			RHICmdList.UnlockBuffer(VertexBufferRHI);
+			const FRHIBufferCreateDesc CreateDesc =
+				FRHIBufferCreateDesc::CreateVertex<uint32>(TEXT("FDummyCulledDispatchVertexIdsBuffer"), 4)
+				.AddUsage(EBufferUsageFlags::Static | EBufferUsageFlags::ShaderResource)
+				.SetInitialState(ERHIAccess::VertexOrIndexBuffer | ERHIAccess::SRVMask)
+				.SetClassName(ClassName)
+				.SetInitActionZeroData();
+
+			VertexBufferRHI = RHICmdList.CreateBuffer(CreateDesc);
 		}
 
 		{
-			FRHIResourceCreateInfo CreateInfo(TEXT("FDummyByteAddressBuffer"));
-			CreateInfo.ClassName = ClassName;
-			uint32 NumBytes = sizeof(uint32) * 4;
-			ByteAddressBufferRHI = RHICmdList.CreateBuffer(NumBytes, BUF_Static | BUF_ShaderResource |BUF_ByteAddressBuffer, 0, ERHIAccess::SRVMask, CreateInfo);
-			uint32* DummyContents = (uint32*)RHICmdList.LockBuffer(ByteAddressBufferRHI, 0, NumBytes, RLM_WriteOnly);
-			DummyContents[0] = DummyContents[1] = DummyContents[2] = DummyContents[3] = 0;
-			RHICmdList.UnlockBuffer(ByteAddressBufferRHI);
+			const FRHIBufferCreateDesc CreateDesc =
+				FRHIBufferCreateDesc::CreateByteAddress(TEXT("FDummyByteAddressBuffer"), sizeof(uint32) * 4, 0)
+				.AddUsage(EBufferUsageFlags::Static | EBufferUsageFlags::ShaderResource)
+				.SetInitialState(ERHIAccess::SRVMask)
+				.SetClassName(ClassName)
+				.SetInitActionZeroData();
+
+			ByteAddressBufferRHI = RHICmdList.CreateBuffer(CreateDesc);
 		}
 
 		SRVUint = RHICmdList.CreateShaderResourceView(
diff --git a/Engine/Plugins/Runtime/OpenXR/Source/OpenXRHMD/Private/OpenXRHMD.cpp b/Engine/Plugins/Runtime/OpenXR/Source/OpenXRHMD/Private/OpenXRHMD.cpp
--- a/Engine/Plugins/Runtime/OpenXR/Source/OpenXRHMD/Private/OpenXRHMD.cpp
+++ b/Engine/Plugins/Runtime/OpenXR/Source/OpenXRHMD/Private/OpenXRHMD.cpp
@@ -1975,24 +1975,32 @@ bool FOpenXRHMD::BuildOcclusionMesh(XrVisibilityMaskTypeKHR Type, int View, FHMD
 
 	FRHICommandListBase& RHICmdList = FRHICommandListImmediate::Get();
 
-	FRHIResourceCreateInfo CreateInfo(TEXT("FOpenXRHMD"));
-	Mesh.VertexBufferRHI = RHICmdList.CreateVertexBuffer(sizeof(FFilterVertex) * VisibilityMask.vertexCountOutput, BUF_Static, CreateInfo);
-	void* VertexBufferPtr = RHICmdList.LockBuffer(Mesh.VertexBufferRHI, 0, sizeof(FFilterVertex) * VisibilityMask.vertexCountOutput, RLM_WriteOnly);
-	FFilterVertex* Vertices = reinterpret_cast<FFilterVertex*>(VertexBufferPtr);
+	const FRHIBufferCreateDesc VertexCreateDesc =
+		FRHIBufferCreateDesc::CreateVertex<FFilterVertex>(TEXT("FOpenXRHMD"), VisibilityMask.vertexCountOutput)
+		.AddUsage(EBufferUsageFlags::Static)
+		.SetInitActionInitializer()
+		.DetermineInitialState();
+	TRHIBufferInitializer<FFilterVertex> VertexInitialData = RHICmdList.CreateBufferInitializer(VertexCreateDesc);
+
+	const FRHIBufferCreateDesc IndexCreateDesc =
+		FRHIBufferCreateDesc::CreateIndex<uint32>(TEXT("FOpenXRHMD"), VisibilityMask.indexCountOutput)
+		.AddUsage(EBufferUsageFlags::Static)
+		.SetInitActionInitializer()
+		.DetermineInitialState();
+	TRHIBufferInitializer<uint32> IndexInitialData = RHICmdList.CreateBufferInitializer(IndexCreateDesc);
 
-	Mesh.IndexBufferRHI = RHICmdList.CreateIndexBuffer(sizeof(uint32), sizeof(uint32) * VisibilityMask.indexCountOutput, BUF_Static, CreateInfo);
-	void* IndexBufferPtr = RHICmdList.LockBuffer(Mesh.IndexBufferRHI, 0, sizeof(uint32) * VisibilityMask.indexCountOutput, RLM_WriteOnly);
-
-	uint32* OutIndices = reinterpret_cast<uint32*>(IndexBufferPtr);
 	TUniquePtr<XrVector2f[]> const OutVertices = MakeUnique<XrVector2f[]>(VisibilityMask.vertexCountOutput);
+	TUniquePtr<uint32[]> const OutIndices = MakeUnique<uint32[]>(VisibilityMask.indexCountOutput);
 
 	VisibilityMask.vertexCapacityInput = VisibilityMask.vertexCountOutput;
 	VisibilityMask.indexCapacityInput = VisibilityMask.indexCountOutput;
-	VisibilityMask.indices = OutIndices;
+	VisibilityMask.indices = OutIndices.Get();
 	VisibilityMask.vertices = OutVertices.Get();
 
 	GetVisibilityMaskKHR(Session, SelectedViewConfigurationType, View, Type, &VisibilityMask);
 
+	IndexInitialData.WriteData(OutIndices.Get(), VisibilityMask.indexCountOutput);
+
 	// We need to apply the eye's projection matrix to each vertex
 	FMatrix Projection = GetStereoProjectionMatrix(View);
 
@@ -2001,7 +2009,7 @@ bool FOpenXRHMD::BuildOcclusionMesh(XrVisibilityMaskTypeKHR Type, int View, FHMD
 
 	for (uint32 VertexIndex = 0; VertexIndex < VisibilityMask.vertexCountOutput; ++VertexIndex)
 	{
-		FFilterVertex& Vertex = Vertices[VertexIndex];
+		FFilterVertex Vertex = { FVector4f(ForceInitToZero), FVector2f(ForceInitToZero) };
 		FVector Position(OutVertices[VertexIndex].x, OutVertices[VertexIndex].y, 1.0f);
 
 		Vertex.Position = (FVector4f)Projection.TransformPosition(Position); // LWC_TODO: precision loss
@@ -2018,15 +2026,17 @@ bool FOpenXRHMD::BuildOcclusionMesh(XrVisibilityMaskTypeKHR Type, int View, FHMD
 
 		Vertex.UV.X = Vertex.Position.X;
 		Vertex.UV.Y = Vertex.Position.Y;
+
+		VertexInitialData.WriteValueAtIndex(VertexIndex, Vertex);
 	}
 
+	Mesh.VertexBufferRHI = VertexInitialData.Finalize();
+	Mesh.IndexBufferRHI = IndexInitialData.Finalize();
+
 	Mesh.NumIndices = VisibilityMask.indexCountOutput;
 	Mesh.NumVertices = VisibilityMask.vertexCountOutput;
 	Mesh.NumTriangles = Mesh.NumIndices / 3;
 
-	RHICmdList.UnlockBuffer(Mesh.VertexBufferRHI);
-	RHICmdList.UnlockBuffer(Mesh.IndexBufferRHI);
-
 	return true;
 }
 #endif
diff --git a/Engine/Plugins/Runtime/WarpUtils/Source/PFMExporter/Private/PFMExporterShader.cpp b/Engine/Plugins/Runtime/WarpUtils/Source/PFMExporter/Private/PFMExporterShader.cpp
--- a/Engine/Plugins/Runtime/WarpUtils/Source/PFMExporter/Private/PFMExporterShader.cpp
+++ b/Engine/Plugins/Runtime/WarpUtils/Source/PFMExporter/Private/PFMExporterShader.cpp
@@ -109,19 +109,24 @@ bool FPFMExporterShader::ApplyPFMExporter_RenderThread(
 	// Initialize data for PFm (RT textures, etc)
 	DstPfmMesh.BeginExport_RenderThread(RHICmdList);
 
-	FRHIResourceCreateInfo CreateInfo(TEXT("FPFMExporterShader"));
-	FBufferRHIRef VertexBufferRHI = RHICmdList.CreateVertexBuffer(sizeof(FFilterVertex) * NumVertices, BUF_Dynamic, CreateInfo);
-	{//Fill buffer with vertex+selected UV channel:
-		void* VoidPtr = RHICmdList.LockBuffer(VertexBufferRHI, 0, sizeof(FFilterVertex) * NumVertices, RLM_WriteOnly);
-		FFilterVertex* pVertices = reinterpret_cast<FFilterVertex*>(VoidPtr);
+	const FRHIBufferCreateDesc CreateDesc =
+		FRHIBufferCreateDesc::CreateVertex<FFilterVertex>(TEXT("FPFMExporterShader"), NumVertices)
+		.AddUsage(EBufferUsageFlags::Dynamic)
+		.SetInitActionInitializer()
+		.DetermineInitialState();
+
+	TRHIBufferInitializer<FFilterVertex> Initializer = RHICmdList.CreateBufferInitializer(CreateDesc);
+	{
+		//Fill buffer with vertex+selected UV channel:
 		for (uint32 i = 0; i < NumVertices; i++)
 		{
-			FFilterVertex& Vertex = pVertices[i];
+			FFilterVertex& Vertex = Initializer[i];
 			Vertex.Position = VertexPosition.VertexPosition(i);
-				Vertex.UV = VertexBuffer.GetVertexUV(i, UVIndex); // Get UV from selected channel
+			Vertex.UV = VertexBuffer.GetVertexUV(i, UVIndex); // Get UV from selected channel
 		}
-		RHICmdList.UnlockBuffer(VertexBufferRHI);
 	}
+	FBufferRHIRef VertexBufferRHI = Initializer.Finalize();
+
 	FRHIBuffer* IndexBufferRHI = IndexBuffer.IndexBufferRHI;
 
 	// Build transform matrix for mesh:
diff --git a/Engine/Plugins/Runtime/WindowsMoviePlayer/Source/WindowsMoviePlayer/Private/WindowsMovieStreamer.cpp b/Engine/Plugins/Runtime/WindowsMoviePlayer/Source/WindowsMoviePlayer/Private/WindowsMovieStreamer.cpp
--- a/Engine/Plugins/Runtime/WindowsMoviePlayer/Source/WindowsMoviePlayer/Private/WindowsMovieStreamer.cpp
+++ b/Engine/Plugins/Runtime/WindowsMoviePlayer/Source/WindowsMoviePlayer/Private/WindowsMovieStreamer.cpp
@@ -156,7 +156,7 @@ void FMediaFoundationMovieStreamer::ConvertSample(FRHICommandListImmediate& RHIC
 		}
 
 		// draw full size quad into render target
-		FBufferRHIRef VertexBuffer = CreateTempMediaVertexBuffer();
+		FBufferRHIRef VertexBuffer = CreateTempMediaVertexBuffer(RHICmdList);
 		RHICmdList.SetStreamSource(0, VertexBuffer, 0);
 		// set viewport to RT size
 		RHICmdList.SetViewport(0, 0, 0.0f, OutputDim.X, OutputDim.Y, 1.0f);
diff --git a/Engine/Plugins/Runtime/nDisplay/Source/DisplayCluster/Private/Render/Containers/DisplayClusterRender_MeshComponentProxy.cpp b/Engine/Plugins/Runtime/nDisplay/Source/DisplayCluster/Private/Render/Containers/DisplayClusterRender_MeshComponentProxy.cpp
--- a/Engine/Plugins/Runtime/nDisplay/Source/DisplayCluster/Private/Render/Containers/DisplayClusterRender_MeshComponentProxy.cpp
+++ b/Engine/Plugins/Runtime/nDisplay/Source/DisplayCluster/Private/Render/Containers/DisplayClusterRender_MeshComponentProxy.cpp
@@ -5,6 +5,7 @@
 #include "Render/Containers/DisplayClusterRender_MeshResources.h"
 
 #include "Misc/DisplayClusterLog.h"
+#include "RHIResourceUtils.h"
 
 TGlobalResource<FDisplayClusterMeshVertexDeclaration> GDisplayClusterMeshVertexDeclaration;
 
@@ -87,36 +88,33 @@ void FDisplayClusterRender_MeshComponentProxy::UpdateRHI_RenderThread(FRHIComman
 				UE_LOG(LogDisplayClusterRender, Warning, TEXT("MeshComponent has a vertex size of 0, please make sure a mesh is assigned."))
 				return;
 			}
-		
-			FRHIResourceCreateInfo CreateInfo(TEXT("DisplayClusterRender_MeshComponentProxy_VertexBuffer"));
-			VertexBufferRHI = RHICmdList.CreateVertexBuffer(VertexDataSize, Usage, CreateInfo);
 
-			FDisplayClusterMeshVertexType* DestVertexData = reinterpret_cast<FDisplayClusterMeshVertexType*>(RHICmdList.LockBuffer(VertexBufferRHI, 0, VertexDataSize, RLM_WriteOnly));
-			if (DestVertexData)
+			const FRHIBufferCreateDesc CreateDesc =
+				FRHIBufferCreateDesc::CreateVertex(TEXT("DisplayClusterRender_MeshComponentProxy_VertexBuffer"), VertexDataSize)
+				.AddUsage(Usage)
+				.SetInitActionInitializer()
+				.DetermineInitialState();
+
+			TRHIBufferInitializer<FDisplayClusterMeshVertexType> DestVertexData = RHICmdList.CreateBufferInitializer(CreateDesc);
 			{
 				const FDisplayClusterMeshVertex* SrcVertexData = InMeshData->GetVertexData().GetData();
 				for (uint32 VertexIdx = 0; VertexIdx < NumVertices; VertexIdx++)
 				{
 					DestVertexData[VertexIdx].SetVertexData(SrcVertexData[VertexIdx]);
 				}
-			
-				RHICmdList.UnlockBuffer(VertexBufferRHI);
 			}
+
+			VertexBufferRHI = DestVertexData.Finalize();
 		}
 
 		// Create Index buffer RHI:
 		{
-			size_t IndexDataSize = sizeof(uint32) * InMeshData->GetIndexData().Num();
-
-			FRHIResourceCreateInfo CreateInfo(TEXT("DisplayClusterRender_MeshComponentProxy_IndexBuffer"));
-			IndexBufferRHI = RHICmdList.CreateIndexBuffer(sizeof(uint32), IndexDataSize, Usage, CreateInfo);
-
-			uint32* DestIndexData = reinterpret_cast<uint32*>(RHICmdList.LockBuffer(IndexBufferRHI, 0, IndexDataSize, RLM_WriteOnly));
-			if(DestIndexData)
-			{
-				FPlatformMemory::Memcpy(DestIndexData, InMeshData->GetIndexData().GetData(), IndexDataSize);
-				RHICmdList.UnlockBuffer(IndexBufferRHI);
-			}
+			IndexBufferRHI = UE::RHIResourceUtils::CreateIndexBufferFromArray(
+				RHICmdList,
+				TEXT("DisplayClusterRender_MeshComponentProxy_VertexBuffer"),
+				Usage,
+				MakeConstArrayView(InMeshData->GetIndexData())
+			);
 		}
 	}
 }
diff --git a/Engine/Plugins/Tests/RHITests/Source/RHITests/Private/RHIDrawTests.cpp b/Engine/Plugins/Tests/RHITests/Source/RHITests/Private/RHIDrawTests.cpp
--- a/Engine/Plugins/Tests/RHITests/Source/RHITests/Private/RHIDrawTests.cpp
+++ b/Engine/Plugins/Tests/RHITests/Source/RHITests/Private/RHIDrawTests.cpp
@@ -106,8 +106,12 @@ struct FDrawTestResources
 
 		static constexpr uint32 OutputBufferStride = sizeof(uint32);
 		static constexpr uint32 OutputBufferSize = OutputBufferStride * MaxInstances;
-		FRHIResourceCreateInfo OutputBufferCreateInfo(TEXT("DrawTest_OutputBuffer"));
-		OutputBuffer = RHICmdList.CreateBuffer(OutputBufferSize, EBufferUsageFlags::UnorderedAccess | EBufferUsageFlags::SourceCopy, OutputBufferStride, ERHIAccess::UAVCompute, OutputBufferCreateInfo);
+		constexpr EBufferUsageFlags OutputBufferUsage = EBufferUsageFlags::UnorderedAccess | EBufferUsageFlags::SourceCopy;
+
+		const FRHIBufferCreateDesc CreateDesc =
+			FRHIBufferCreateDesc::Create(TEXT("DrawTest_OutputBuffer"), OutputBufferSize, OutputBufferStride, OutputBufferUsage)
+			.SetInitialState(ERHIAccess::UAVCompute);
+		OutputBuffer = RHICmdList.CreateBuffer(CreateDesc);
 
 		OutputBufferUAV = RHICmdList.CreateUnorderedAccessView(OutputBuffer,
 			FRHIViewDesc::CreateBufferUAV()
diff --git a/Engine/Plugins/Tests/RHITests/Source/RHITests/Private/RHIGraphicsUAVTests.cpp b/Engine/Plugins/Tests/RHITests/Source/RHITests/Private/RHIGraphicsUAVTests.cpp
--- a/Engine/Plugins/Tests/RHITests/Source/RHITests/Private/RHIGraphicsUAVTests.cpp
+++ b/Engine/Plugins/Tests/RHITests/Source/RHITests/Private/RHIGraphicsUAVTests.cpp
@@ -120,12 +120,13 @@ bool FRHIGraphicsUAVTests::Test_GraphicsUAV_PixelShader(FRHICommandListImmediate
 	static constexpr uint32 MaxInstances = 8;
 	static constexpr uint32 OutputBufferStride = sizeof(uint32);
 	static constexpr uint32 OutputBufferSize = OutputBufferStride * MaxInstances;
-	FRHIResourceCreateInfo OutputBufferCreateInfo(TEXT("GraphicsUAVTests_PixelShaderOutput"));
+	constexpr EBufferUsageFlags OutputBufferUsage = EBufferUsageFlags::UnorderedAccess | EBufferUsageFlags::SourceCopy | EBufferUsageFlags::StructuredBuffer;
 
 	// NOTE: using a structured buffer here as workaround for UE-212251
-	FBufferRHIRef OutputBuffer = RHICmdList.CreateBuffer(OutputBufferSize,
-		EBufferUsageFlags::UnorderedAccess | EBufferUsageFlags::SourceCopy | EBufferUsageFlags::StructuredBuffer,
-		OutputBufferStride, ERHIAccess::UAVCompute, OutputBufferCreateInfo);
+	const FRHIBufferCreateDesc CreateDesc =
+		FRHIBufferCreateDesc::Create(TEXT("GraphicsUAVTests_PixelShaderOutput"), OutputBufferSize, OutputBufferStride, OutputBufferUsage)
+		.SetInitialState(ERHIAccess::UAVCompute);
+	FBufferRHIRef OutputBuffer = RHICmdList.CreateBuffer(CreateDesc);
 
 	FUnorderedAccessViewRHIRef OutputBufferUAV = RHICmdList.CreateUnorderedAccessView(OutputBuffer,
 		FRHIViewDesc::CreateBufferUAV()
@@ -170,12 +171,14 @@ bool FRHIGraphicsUAVTests::Test_GraphicsUAV_VertexShader(FRHICommandListImmediat
 	static const uint32 MaxVertices = 3;
 	static constexpr uint32 OutputBufferStride = sizeof(uint32);
 	static constexpr uint32 OutputBufferSize = OutputBufferStride * MaxVertices;
-	FRHIResourceCreateInfo OutputBufferCreateInfo(TEXT("GraphicsUAVTests_VertexShaderOutput"));
+	constexpr EBufferUsageFlags OutputBufferUsage = EBufferUsageFlags::UnorderedAccess | EBufferUsageFlags::SourceCopy | EBufferUsageFlags::StructuredBuffer;
 
 	// NOTE: using a structured buffer here as workaround for UE-212251
-	FBufferRHIRef OutputBuffer = RHICmdList.CreateBuffer(OutputBufferSize,
-		EBufferUsageFlags::UnorderedAccess | EBufferUsageFlags::SourceCopy | EBufferUsageFlags::StructuredBuffer,
-		OutputBufferStride, ERHIAccess::UAVCompute, OutputBufferCreateInfo);
+	const FRHIBufferCreateDesc CreateDesc =
+		FRHIBufferCreateDesc::Create(TEXT("GraphicsUAVTests_VertexShaderOutput"), OutputBufferSize, OutputBufferStride, OutputBufferUsage)
+		.SetInitialState(ERHIAccess::UAVCompute);
+
+	FBufferRHIRef OutputBuffer = RHICmdList.CreateBuffer(CreateDesc);
 
 	FUnorderedAccessViewRHIRef OutputBufferUAV = RHICmdList.CreateUnorderedAccessView(OutputBuffer,
 		FRHIViewDesc::CreateBufferUAV()
diff --git a/Engine/Plugins/Tests/RHITests/Source/RHITests/Private/RHIReadbackTests.cpp b/Engine/Plugins/Tests/RHITests/Source/RHITests/Private/RHIReadbackTests.cpp
--- a/Engine/Plugins/Tests/RHITests/Source/RHITests/Private/RHIReadbackTests.cpp
+++ b/Engine/Plugins/Tests/RHITests/Source/RHITests/Private/RHIReadbackTests.cpp
@@ -77,9 +77,12 @@ bool FRHIReadbackTests::Test_BufferReadback(FRHICommandListImmediate& RHICmdList
 	// Create typed, byte address and structured buffers.
 	//
 	{
-		FRHIResourceCreateInfo TypedBufferCreateInfo(TEXT("TypedBuffer"));
-		EBufferUsageFlags TypedBufferUsage = EBufferUsageFlags::UnorderedAccess | EBufferUsageFlags::SourceCopy;
-		TypedBuffer = RHICmdList.CreateBuffer(TypedBufferSize, TypedBufferUsage, 0, ERHIAccess::UAVCompute, TypedBufferCreateInfo);
+		const EBufferUsageFlags TypedBufferUsage = EBufferUsageFlags::UnorderedAccess | EBufferUsageFlags::SourceCopy;
+		const FRHIBufferCreateDesc CreateDesc =
+			FRHIBufferCreateDesc::Create(TEXT("TypedBuffer"), TypedBufferSize, 0, TypedBufferUsage)
+			.SetInitialState(ERHIAccess::UAVCompute);
+		TypedBuffer = RHICmdList.CreateBuffer(CreateDesc);
+
 		auto TypedBufferUAVCreateDesc = FRHIViewDesc::CreateBufferUAV()
 			.SetType(FRHIViewDesc::EBufferType::Typed)
 			.SetFormat(PF_R32_UINT)
@@ -88,18 +91,24 @@ bool FRHIReadbackTests::Test_BufferReadback(FRHICommandListImmediate& RHICmdList
 	}
 
 	{
-		FRHIResourceCreateInfo ByteAddressBufferCreateInfo(TEXT("ByteAddressBuffer"));
-		EBufferUsageFlags ByteAddressBufferUsage = EBufferUsageFlags::UnorderedAccess | EBufferUsageFlags::SourceCopy | EBufferUsageFlags::ByteAddressBuffer | EBufferUsageFlags::StructuredBuffer;
-		ByteAddressBuffer = RHICmdList.CreateBuffer(ByteAddressBufferSize, ByteAddressBufferUsage, sizeof(uint32), ERHIAccess::UAVCompute, ByteAddressBufferCreateInfo);
+		const EBufferUsageFlags ByteAddressBufferUsage = EBufferUsageFlags::UnorderedAccess | EBufferUsageFlags::SourceCopy | EBufferUsageFlags::ByteAddressBuffer | EBufferUsageFlags::StructuredBuffer;
+		const FRHIBufferCreateDesc CreateDesc =
+			FRHIBufferCreateDesc::Create(TEXT("ByteAddressBuffer"), ByteAddressBufferSize, sizeof(uint32), ByteAddressBufferUsage)
+			.SetInitialState(ERHIAccess::UAVCompute);
+		ByteAddressBuffer = RHICmdList.CreateBuffer(CreateDesc);
+
 		auto ByteAddressBufferUAVCreateDesc = FRHIViewDesc::CreateBufferUAV()
 			.SetType(FRHIViewDesc::EBufferType::Raw);
 		ByteAddressBufferUAV = RHICmdList.CreateUnorderedAccessView(ByteAddressBuffer, ByteAddressBufferUAVCreateDesc);
 	}
 
 	{
-		FRHIResourceCreateInfo StructuredBufferCreateInfo(TEXT("StructuredBuffer"));
-		EBufferUsageFlags StructuredBufferUsage = EBufferUsageFlags::UnorderedAccess | EBufferUsageFlags::SourceCopy | EBufferUsageFlags::StructuredBuffer;
-		StructuredBuffer = RHICmdList.CreateBuffer(StructuredBufferSize, StructuredBufferUsage, sizeof(FStructuredBufferElem), ERHIAccess::UAVCompute, StructuredBufferCreateInfo);
+		const EBufferUsageFlags StructuredBufferUsage = EBufferUsageFlags::UnorderedAccess | EBufferUsageFlags::SourceCopy | EBufferUsageFlags::StructuredBuffer;
+		const FRHIBufferCreateDesc CreateDesc =
+			FRHIBufferCreateDesc::Create(TEXT("StructuredBuffer"), StructuredBufferSize, sizeof(FStructuredBufferElem), StructuredBufferUsage)
+			.SetInitialState(ERHIAccess::UAVCompute);
+		StructuredBuffer = RHICmdList.CreateBuffer(CreateDesc);
+
 		auto StructuredBufferUAVCreateDesc = FRHIViewDesc::CreateBufferUAV()
 			.SetType(FRHIViewDesc::EBufferType::Structured)
 			.SetStride(sizeof(FStructuredBufferElem))
diff --git a/Engine/Plugins/Tests/RHITests/Source/RHITests/Private/RHIReservedResourceTests.cpp b/Engine/Plugins/Tests/RHITests/Source/RHITests/Private/RHIReservedResourceTests.cpp
--- a/Engine/Plugins/Tests/RHITests/Source/RHITests/Private/RHIReservedResourceTests.cpp
+++ b/Engine/Plugins/Tests/RHITests/Source/RHITests/Private/RHIReservedResourceTests.cpp
@@ -215,28 +215,43 @@ bool FRHIReservedResourceTests::Test_ReservedResource_CreateBuffer(FRHICommandLi
 	// Simply try to create reserved buffers of different types and sizes to see if we hit any unexpected paths in the RHI
 
 	{
-		FRHIResourceCreateInfo CreateInfo(TEXT("TestSmallReservedVertexBuffer"));
-		FBufferRHIRef Buffer = RHICmdList.CreateBuffer(32768, BUF_ReservedResource | BUF_VertexBuffer, 4, ERHIAccess::CopyDest, CreateInfo);
+		const FRHIBufferCreateDesc CreateDesc =
+			FRHIBufferCreateDesc::CreateVertex(TEXT("TestSmallReservedVertexBuffer"), 32768)
+			.AddUsage(EBufferUsageFlags::ReservedResource)
+			.SetInitialState(ERHIAccess::CopyDest);
+		FBufferRHIRef Buffer = RHICmdList.CreateBuffer(CreateDesc);
 	}
 
 	{
-		FRHIResourceCreateInfo CreateInfo(TEXT("TestSmallReservedUAV"));
-		FBufferRHIRef Buffer = RHICmdList.CreateBuffer(32768, BUF_ReservedResource | BUF_UnorderedAccess | BUF_ShaderResource, 4, ERHIAccess::UAVGraphics, CreateInfo);
+		const FRHIBufferCreateDesc CreateDesc =
+			FRHIBufferCreateDesc::Create(TEXT("TestSmallReservedUAV"), 32768, 4, EBufferUsageFlags::UnorderedAccess | EBufferUsageFlags::ShaderResource)
+			.AddUsage(EBufferUsageFlags::ReservedResource)
+			.SetInitialState(ERHIAccess::UAVGraphics);
+		FBufferRHIRef Buffer = RHICmdList.CreateBuffer(CreateDesc);
 	}
 
 	{
-		FRHIResourceCreateInfo CreateInfo(TEXT("TestReservedVertexBuffer"));
-		FBufferRHIRef Buffer = RHICmdList.CreateBuffer(32 * 1024 * 1024, BUF_ReservedResource | BUF_VertexBuffer, 4, ERHIAccess::CopyDest, CreateInfo);
+		const FRHIBufferCreateDesc CreateDesc =
+			FRHIBufferCreateDesc::CreateVertex(TEXT("TestReservedVertexBuffer"), 32 * 1024 * 1024)
+			.AddUsage(EBufferUsageFlags::ReservedResource)
+			.SetInitialState(ERHIAccess::CopyDest);
+		FBufferRHIRef Buffer = RHICmdList.CreateBuffer(CreateDesc);
 	}
 
 	{
-		FRHIResourceCreateInfo CreateInfo(TEXT("TestReservedAccelerationStructureBuffer"));
-		FBufferRHIRef Buffer = RHICmdList.CreateBuffer(32 * 1024 * 1024, BUF_ReservedResource | BUF_AccelerationStructure, 4, ERHIAccess::BVHWrite, CreateInfo);
+		const FRHIBufferCreateDesc CreateDesc =
+			FRHIBufferCreateDesc::Create(TEXT("TestReservedAccelerationStructureBuffer"), 32 * 1024 * 1024, 4, EBufferUsageFlags::AccelerationStructure)
+			.SetUsage(EBufferUsageFlags::ReservedResource)
+			.SetInitialState(ERHIAccess::BVHWrite);
+		FBufferRHIRef Buffer = RHICmdList.CreateBuffer(CreateDesc);
 	}
 
 	{
-		FRHIResourceCreateInfo CreateInfo(TEXT("TestReservedRayTracingScratchBuffer"));
-		FBufferRHIRef Buffer = RHICmdList.CreateBuffer(32 * 1024 * 1024, BUF_ReservedResource | BUF_RayTracingScratch, 4, ERHIAccess::UAVCompute, CreateInfo);
+		const FRHIBufferCreateDesc CreateDesc =
+			FRHIBufferCreateDesc::Create(TEXT("TestReservedRayTracingScratchBuffer"), 32 * 1024 * 1024, 4, EBufferUsageFlags::RayTracingScratch)
+			.AddUsage(EBufferUsageFlags::ReservedResource)
+			.SetInitialState(ERHIAccess::UAVCompute);
+		FBufferRHIRef Buffer = RHICmdList.CreateBuffer(CreateDesc);
 	}
 
 	return true;
@@ -313,11 +328,11 @@ bool FRHIReservedResourceTests::Test_ReservedResource_CommitBuffer(FRHICommandLi
 	const int32 TileSizeInBytes = GRHIGlobals.ReservedResources.TileSizeInBytes;
 	const int32 BufferSizeInBytes = TileSizeInBytes * 128;
 
-	FRHIResourceCreateInfo CreateInfo(TEXT("TestReservedBufferExplicitCommit"));
+	const FRHIBufferCreateDesc CreateDesc =
+		FRHIBufferCreateDesc::Create(TEXT("TestReservedBufferExplicitCommit"), BufferSizeInBytes, 4, EBufferUsageFlags::ReservedResource | EBufferUsageFlags::UnorderedAccess | EBufferUsageFlags::ShaderResource | EBufferUsageFlags::SourceCopy)
+		.SetInitialState(ERHIAccess::UAVCompute);
 
-	FBufferRHIRef Buffer = RHICmdList.CreateBuffer(BufferSizeInBytes,
-		BUF_ReservedResource | BUF_UnorderedAccess | BUF_ShaderResource | BUF_SourceCopy,
-		4, ERHIAccess::UAVCompute, CreateInfo);
+	FBufferRHIRef Buffer = RHICmdList.CreateBuffer(CreateDesc);
 
 	FUnorderedAccessViewRHIRef BufferUAV = RHICmdList.CreateUnorderedAccessView(Buffer,
 		FRHIViewDesc::CreateBufferUAV()
@@ -356,11 +371,11 @@ bool FRHIReservedResourceTests::Test_ReservedResource_DecommitBuffer(FRHICommand
 	const int32 TileSizeInBytes = GRHIGlobals.ReservedResources.TileSizeInBytes;
 	const int32 BufferSizeInBytes = TileSizeInBytes * 128;
 
-	FRHIResourceCreateInfo CreateInfo(TEXT("TestReservedBufferExplicitCommit"));
+	const FRHIBufferCreateDesc CreateDesc =
+		FRHIBufferCreateDesc::Create(TEXT("TestReservedBufferExplicitCommit"), BufferSizeInBytes, 4, EBufferUsageFlags::ReservedResource | EBufferUsageFlags::UnorderedAccess | EBufferUsageFlags::ShaderResource | EBufferUsageFlags::SourceCopy)
+		.SetInitialState(ERHIAccess::UAVCompute);
 
-	FBufferRHIRef Buffer = RHICmdList.CreateBuffer(BufferSizeInBytes,
-		BUF_ReservedResource | BUF_UnorderedAccess | BUF_ShaderResource | BUF_SourceCopy,
-		4, ERHIAccess::UAVCompute, CreateInfo);
+	FBufferRHIRef Buffer = RHICmdList.CreateBuffer(CreateDesc);
 
 	FUnorderedAccessViewRHIRef BufferUAV = RHICmdList.CreateUnorderedAccessView(Buffer,
 		FRHIViewDesc::CreateBufferUAV()
diff --git a/Engine/Plugins/Tests/RHITests/Source/RHITests/Private/RHITests.spec.cpp b/Engine/Plugins/Tests/RHITests/Source/RHITests/Private/RHITests.spec.cpp
--- a/Engine/Plugins/Tests/RHITests/Source/RHITests/Private/RHITests.spec.cpp
+++ b/Engine/Plugins/Tests/RHITests/Source/RHITests/Private/RHITests.spec.cpp
@@ -109,6 +109,15 @@ void FAutomationRHITest::Define()
 		});
 	});
 
+	Describe("Test RHI Create Buffer", [this]
+	{
+		It("RHICreateBuffer", [this]()
+		{
+			bool bResult = RunOnRenderThreadSynchronous(FRHIBufferTests::Test_RHICreateBuffer);
+			TestEqual("Test_RHICreateBuffer", bResult, 1);
+		});
+	});
+
 	Describe("Test RHI Create Buffer Parallel", [this]
 	{
 		It("RHICreateBuffer_Parallel", [this]()
diff --git a/Engine/Plugins/Tests/RHITests/Source/RHITests/Private/RHITestsModule.cpp b/Engine/Plugins/Tests/RHITests/Source/RHITests/Private/RHITestsModule.cpp
--- a/Engine/Plugins/Tests/RHITests/Source/RHITests/Private/RHITestsModule.cpp
+++ b/Engine/Plugins/Tests/RHITests/Source/RHITests/Private/RHITestsModule.cpp
@@ -85,6 +85,7 @@ static bool RunTests_RenderThread(FRHICommandListImmediate& RHICmdList)
 	// Buffer Operations
 	// ------------------------------------------------
 	{
+		RUN_TEST(FRHIBufferTests::Test_RHICreateBuffer(RHICmdList));
 		RUN_TEST(FRHIBufferTests::Test_RHICreateBuffer_Parallel(RHICmdList));
 	}
 
diff --git a/Engine/Plugins/Tests/RHITests/Source/RHITests/Public/RHIBufferTests.h b/Engine/Plugins/Tests/RHITests/Source/RHITests/Public/RHIBufferTests.h
--- a/Engine/Plugins/Tests/RHITests/Source/RHITests/Public/RHIBufferTests.h
+++ b/Engine/Plugins/Tests/RHITests/Source/RHITests/Public/RHIBufferTests.h
@@ -168,8 +168,11 @@ private:
 			return true;
 		}
 
-		FRHIResourceCreateInfo Info(*TestName);
-		FBufferRHIRef Buffer = RHICmdList.CreateBuffer(BufferSize, BufferFlags | BUF_ShaderResource | BUF_UnorderedAccess | BUF_SourceCopy, Stride, ERHIAccess::UAVMask, Info);
+		const FRHIBufferCreateDesc Desc =
+			FRHIBufferCreateDesc::Create(*TestName, BufferSize, Stride, BufferFlags | BUF_ShaderResource | BUF_UnorderedAccess | BUF_SourceCopy)
+			.SetInitialState(ERHIAccess::UAVMask);
+
+		FBufferRHIRef Buffer = RHICmdList.CreateBuffer(Desc);
 
 		auto UAVCreateDesc = FRHIViewDesc::CreateBufferUAV();
 		auto SRVCreateDesc = FRHIViewDesc::CreateBufferSRV();
@@ -401,6 +404,101 @@ public:
 		return bResult;
 	}
 
+	static bool Test_RHICreateBuffer_Buffer(FRHICommandListImmediate& RHICmdList, uint32 BufferSize)
+	{
+		SCOPED_NAMED_EVENT_TEXT("Test_RHICreateBuffer", FColor::Magenta);
+
+		FString TestName = FString::Printf(TEXT("Test_RHICreateBuffer Size: %d"), BufferSize);
+
+		if (BufferSize == 0)
+		{
+			FBufferRHIRef BufferRHI = RHICmdList.CreateBuffer(FRHIBufferCreateDesc::CreateNull(*TestName).SetInitialState(ERHIAccess::VertexOrIndexBuffer));
+			return BufferRHI != nullptr;
+		}
+
+		// Build arrays of the expected buffer values
+		TArray<uint8> Expected_Zero;
+		TArray<uint8> Expected_Offsets;
+		{
+			Expected_Zero.SetNumZeroed(BufferSize);
+			Expected_Offsets.SetNumUninitialized(BufferSize);
+			for (uint32 Index = 0; Index < BufferSize; Index++)
+			{
+				// Rollover is expected here
+				Expected_Offsets[Index] = static_cast<uint8>(Index);
+			}
+		}
+
+		const FRHIBufferCreateDesc CommonCreateDesc =
+			FRHIBufferCreateDesc::Create(*TestName, BufferSize, 0, EBufferUsageFlags::VertexBuffer)
+			// Required for verifying buffer contents
+			.AddUsage(EBufferUsageFlags::SourceCopy)
+			.SetInitialState(ERHIAccess::CopySrc);
+
+		if (false) // check Fatal conditions
+		{
+			FRHIBufferInitializer Initializer = RHICmdList.CreateBufferInitializer(CommonCreateDesc);
+			RHICmdList.ImmediateFlush(EImmediateFlushType::DispatchToRHIThread);
+
+			FBufferRHIRef BufferRHI = Initializer.Finalize();
+		}
+
+		bool bResult = true;
+
+		{
+			// Create Zeroed
+			FRHIBufferCreateDesc CreateDesc =
+				FRHIBufferCreateDesc(CommonCreateDesc)
+				.SetInitActionZeroData();
+
+			FBufferRHIRef BufferRHI = RHICmdList.CreateBuffer(CreateDesc);
+
+			const bool bDataMatches = VerifyBufferContents(*FString::Printf(TEXT("%s - ZeroData"), *TestName), RHICmdList, BufferRHI, Expected_Zero);
+			bResult = bDataMatches ? bResult : false;
+		}
+
+		{
+			// Create using initializer
+			FRHIBufferCreateDesc CreateDesc =
+				FRHIBufferCreateDesc(CommonCreateDesc)
+				.SetInitActionInitializer();
+
+			TRHIBufferInitializer<uint8> Initializer = RHICmdList.CreateBufferInitializer(CreateDesc);
+			Initializer.WriteArray(Expected_Offsets);
+			FBufferRHIRef BufferRHI = Initializer.Finalize();
+
+			const bool bDataMatches = VerifyBufferContents(*FString::Printf(TEXT("%s - Initializer"), *TestName), RHICmdList, BufferRHI, Expected_Offsets);
+			bResult = bDataMatches ? bResult : false;
+		}
+
+		{
+			// Create from array
+			FResourceArrayUploadArrayView UploadView(Expected_Offsets);
+			const FRHIBufferCreateDesc CreateDesc =
+				FRHIBufferCreateDesc(CommonCreateDesc)
+				.SetInitActionResourceArray(&UploadView);
+
+			FBufferRHIRef BufferRHI = RHICmdList.CreateBuffer(CreateDesc);
+
+			const bool bDataMatches = VerifyBufferContents(*FString::Printf(TEXT("%s - ResourceArray"), *TestName), RHICmdList, BufferRHI, Expected_Offsets);
+			bResult = bDataMatches ? bResult : false;
+		}
+
+		return bResult;
+	}
+
+	static bool Test_RHICreateBuffer(FRHICommandListImmediate& RHICmdList)
+	{
+		bool bResult = true;
+
+		RUN_TEST(Test_RHICreateBuffer_Buffer(RHICmdList, 0));
+		RUN_TEST(Test_RHICreateBuffer_Buffer(RHICmdList, 1));
+		RUN_TEST(Test_RHICreateBuffer_Buffer(RHICmdList, 4));
+		RUN_TEST(Test_RHICreateBuffer_Buffer(RHICmdList, 2048));
+
+		return bResult;
+	}
+
 	static bool Test_RHICreateBuffer_Parallel(FRHICommandListImmediate& RHICmdList)
 	{
 		if (!GRHISupportsMultithreadedResources)
@@ -433,37 +531,34 @@ public:
 			const int32 BufferSize = Align(Random, 16);
 			const int32 BufferStride = 4;
 
-			EBufferUsageFlags Usage = EBufferUsageFlags::VertexBuffer | EBufferUsageFlags::SourceCopy;
+			FRHIBufferCreateDesc Desc =
+				FRHIBufferCreateDesc::Create(TEXT("Buffer"), BufferSize, 0, EBufferUsageFlags::VertexBuffer | EBufferUsageFlags::SourceCopy)
+				.SetInitialState(ERHIAccess::CopySrc)
+				.SetInitActionInitializer();
 
 			switch (Random % 3)
 			{
 			case 0:
-				Usage |= EBufferUsageFlags::Static;
+				Desc.AddUsage(EBufferUsageFlags::Static);
 				break;
 			case 1:
-				Usage |= EBufferUsageFlags::Dynamic;
+				Desc.AddUsage(EBufferUsageFlags::Dynamic);
 				break;
 			case 2:
-				Usage |= EBufferUsageFlags::Volatile;
+				Desc.AddUsage(EBufferUsageFlags::Volatile);
 				break;
 			}
 
-			FRHIResourceCreateInfo CreateInfo(TEXT("Buffer"));
+			TRHIBufferInitializer<uint32> Initializer = InRHICmdList.CreateBufferInitializer(Desc);
 
-			TRefCountPtr<FRHIBuffer> Buffer = InRHICmdList.CreateBuffer(BufferSize, Usage, 0, ERHIAccess::CopySrc, CreateInfo);
-
-			uint32* Data = (uint32*)InRHICmdList.LockBuffer(Buffer, 0, Buffer->GetSize(), RLM_WriteOnly);
-
-			uint32 NumDWORDs = Buffer->GetSize() >> 2;
+			uint32 NumDWORDs = BufferSize >> 2;
 
 			for (uint32 DataIndex = 0; DataIndex < NumDWORDs; ++DataIndex)
 			{
-				Data[DataIndex] = DataIndex;
+				Initializer[DataIndex] = DataIndex;
 			}
 
-			InRHICmdList.UnlockBuffer(Buffer);
-
-			Buffers[Index] = Buffer;
+			Buffers[Index] = Initializer.Finalize();
 		});
 
 		TArray<FRHIBuffer*> BufferPtrs;
diff --git a/Engine/Plugins/Tests/RHITests/Source/RHITests/Public/RHITextureTests.h b/Engine/Plugins/Tests/RHITests/Source/RHITests/Public/RHITextureTests.h
--- a/Engine/Plugins/Tests/RHITests/Source/RHITests/Public/RHITextureTests.h
+++ b/Engine/Plugins/Tests/RHITests/Source/RHITests/Public/RHITextureTests.h
@@ -541,11 +541,14 @@ public:
 		}
 		else if (UpdateType == ETextureUpdateType::TextureFromBuffer)
 		{
-			FRHIResourceCreateInfo CreateInfo(TEXT("SourceUpdateBuffer"));
-			TRefCountPtr<FRHIBuffer> Buffer = RHICmdList.CreateBuffer(SourceSize, EBufferUsageFlags::ShaderResource | EBufferUsageFlags::Static | EBufferUsageFlags::KeepCPUAccessible, SourceType::ElementSize, ERHIAccess::SRVCompute, CreateInfo);
-			void* Data = RHICmdList.LockBuffer(Buffer, 0, SourceSize, RLM_WriteOnly);
-			FMemory::Memcpy(Data, SourceData, SourceSize);
-			RHICmdList.UnlockBuffer(Buffer);
+			const FRHIBufferCreateDesc CreateDesc =
+				FRHIBufferCreateDesc::Create(TEXT("SourceUpdateBuffer"), SourceSize, SourceType::ElementSize, EBufferUsageFlags::ShaderResource | EBufferUsageFlags::Static | EBufferUsageFlags::KeepCPUAccessible)
+				.SetInitialState(ERHIAccess::SRVCompute)
+				.SetInitActionInitializer();
+
+			FRHIBufferInitializer Initializer = RHICmdList.CreateBufferInitializer(CreateDesc);
+			Initializer.WriteData(SourceData, SourceSize);
+			FBufferRHIRef Buffer = Initializer.Finalize();
 
 			RHICmdList.UpdateFromBufferTexture2D(Texture, 0, Region, SourcePitch, Buffer, 0);
 		}
diff --git a/Engine/Plugins/TextureGraph/Source/TextureGraphEngine/3D/RenderMesh.cpp b/Engine/Plugins/TextureGraph/Source/TextureGraphEngine/3D/RenderMesh.cpp
--- a/Engine/Plugins/TextureGraph/Source/TextureGraphEngine/3D/RenderMesh.cpp
+++ b/Engine/Plugins/TextureGraph/Source/TextureGraphEngine/3D/RenderMesh.cpp
@@ -11,6 +11,7 @@
 #include "Transform/Mesh/T_MeshEncodedAsTexture.h"
 #include "Transform/Mesh/T_MeshDilateUVs.h"
 #include "Helper/MathUtils.h"
+#include "RHIResourceUtils.h"
 
 #if WITH_EDITOR
 
@@ -285,7 +286,7 @@ void RenderMesh::Render_Now(FRHICommandList& RHICmdList, int32 targetId) const
 		static const size_t maxVertices = 16 * 1024;
 		size_t maxChunks = numVertices / maxVertices;
 
-		TResourceArray<CoreMeshVertex, VERTEXBUFFER_ALIGNMENT> vertices;
+		TArray<CoreMeshVertex> vertices;
 		vertices.SetNumUninitialized(numVertices);
 
 		for (size_t vi = 0; vi < numVertices; vi++)
@@ -312,7 +313,7 @@ void RenderMesh::Render_Now(FRHICommandList& RHICmdList, int32 targetId) const
 			vertex.color = FVector4f(color.R, color.G, color.B, color.A);
 		}
 
-		TResourceArray<uint32, INDEXBUFFER_ALIGNMENT> indices;
+		TArray<uint32> indices;
 		indices.SetNumUninitialized(numIndices);
 
 		for (size_t ii = 0; ii < numIndices; ii++)
@@ -320,9 +321,8 @@ void RenderMesh::Render_Now(FRHICommandList& RHICmdList, int32 targetId) const
 			indices[ii] = cmesh->triangles[ii];
 		}
 
-		FRHIResourceCreateInfo createInfoVB(TEXT("RenderMesh_VB"), &vertices), createInfoIB(TEXT("RenderMesh_IB"), &indices);
-		FBufferRHIRef vertexBuffer = RHICmdList.CreateVertexBuffer(vertexSize * numVertices, BUF_Static, createInfoVB);
-		FBufferRHIRef indexBuffer = RHICmdList.CreateIndexBuffer(0, sizeof(uint32) * numIndices, BUF_Static, createInfoIB);
+		FBufferRHIRef vertexBuffer = UE::RHIResourceUtils::CreateVertexBufferFromArray(RHICmdList, TEXT("RenderMesh_VB"), EBufferUsageFlags::Static, MakeConstArrayView(vertices));
+		FBufferRHIRef indexBuffer = UE::RHIResourceUtils::CreateIndexBufferFromArray(RHICmdList, TEXT("RenderMesh_IB"), EBufferUsageFlags::Static, MakeConstArrayView(indices));
 
 		RHICmdList.SetStreamSource(0, vertexBuffer, 0);
 		RHICmdList.DrawIndexedPrimitive(indexBuffer, 0, 0, numVertices, 0, numTris, 1);
diff --git a/Engine/Plugins/VirtualProduction/CameraCalibrationCore/Source/CameraCalibrationCore/Private/DistortionRenderingUtils.cpp b/Engine/Plugins/VirtualProduction/CameraCalibrationCore/Source/CameraCalibrationCore/Private/DistortionRenderingUtils.cpp
--- a/Engine/Plugins/VirtualProduction/CameraCalibrationCore/Source/CameraCalibrationCore/Private/DistortionRenderingUtils.cpp
+++ b/Engine/Plugins/VirtualProduction/CameraCalibrationCore/Source/CameraCalibrationCore/Private/DistortionRenderingUtils.cpp
@@ -11,6 +11,7 @@
 #include "TextureResource.h"
 #include "DataDrivenShaderPlatformInfo.h"
 #include "RHIUtilities.h"
+#include "RHIResourceUtils.h"
 
 class FUndistortImagePointsCS : public FGlobalShader
 {
@@ -78,21 +79,19 @@ namespace DistortionRenderingUtils
 				Parameters.DistortionMap = DistortionMapResource->TextureRHI;
 				Parameters.DistortionMapSampler = TStaticSamplerState<SF_Bilinear, AM_Clamp, AM_Clamp, AM_Clamp>::GetRHI();
 
-				TResourceArray<FVector2f> InputPointsResourceArray;
-				InputPointsResourceArray.Reserve(NumPoints);
-				for (int32 Index = 0; Index < NumPoints; Index++)
-				{
-					InputPointsResourceArray.Add(Points[Index]);
-				}
-
 				TResourceArray<FVector2f> EmptyBuffer;
 				EmptyBuffer.AddZeroed(NumPoints);
 
 				const uint32 BufferSize = sizeof(FVector2f) * NumPoints;
 
 				// Create an SRV for the input buffer of image points
-				FRHIResourceCreateInfo CreateInfo(TEXT("ImagePointsInitialData"), &InputPointsResourceArray);
-				FBufferRHIRef InputPointsBuffer = RHICmdList.CreateStructuredBuffer(sizeof(FVector2f), BufferSize, BUF_Static | BUF_ShaderResource, CreateInfo);
+				FBufferRHIRef InputPointsBuffer = UE::RHIResourceUtils::CreateBufferFromArray(
+					RHICmdList,
+					TEXT("ImagePointsInitialData"),
+					EBufferUsageFlags::StructuredBuffer | EBufferUsageFlags::Static | EBufferUsageFlags::ShaderResource,
+					MakeConstArrayView(Points)
+				);
+
 				FShaderResourceViewRHIRef InputPointsSRV = RHICmdList.CreateShaderResourceView(InputPointsBuffer, FRHIViewDesc::CreateBufferSRV().SetTypeFromBuffer(InputPointsBuffer));
 				Parameters.InputPoints = InputPointsSRV;
 
diff --git a/Engine/Plugins/VirtualProduction/CameraCalibrationCore/Source/CameraCalibrationCore/Private/LensDistortionSceneViewExtension.cpp b/Engine/Plugins/VirtualProduction/CameraCalibrationCore/Source/CameraCalibrationCore/Private/LensDistortionSceneViewExtension.cpp
--- a/Engine/Plugins/VirtualProduction/CameraCalibrationCore/Source/CameraCalibrationCore/Private/LensDistortionSceneViewExtension.cpp
+++ b/Engine/Plugins/VirtualProduction/CameraCalibrationCore/Source/CameraCalibrationCore/Private/LensDistortionSceneViewExtension.cpp
@@ -7,6 +7,7 @@
 #include "Containers/DynamicRHIResourceArray.h"
 #include "GlobalShader.h"
 #include "PostProcess/LensDistortion.h"
+#include "RHIResourceUtils.h"
 #include "ScreenPass.h"
 #include "ShaderParameterStruct.h"
 #include "SystemTextures.h"
@@ -319,7 +320,7 @@ void FLensDistortionSceneViewExtension::InvertDistortionMap_RenderThread(FRDGBui
 			RHICmdList.SetStreamSource(0, NULL, 0);
 
 			// The following code for setting up this index buffer based on FTesselatedScreenRectangleIndexBuffer::InitRHI()
-			TResourceArray<uint32> IndexBuffer;
+			TArray<uint32> IndexBuffer;
 
 			const uint32 Width = GridDimensions.X;
 			const uint32 Height = GridDimensions.Y;
@@ -353,8 +354,7 @@ void FLensDistortionSceneViewExtension::InvertDistortionMap_RenderThread(FRDGBui
 			}
 
 			// Create index buffer. Fill buffer with initial data upon creation
-			FRHIResourceCreateInfo CreateInfo(TEXT("InvertDistortionMapIndexBuffer"), &IndexBuffer);
-			FBufferRHIRef IndexBufferRHI = RHICmdList.CreateIndexBuffer(sizeof(uint32), IndexBuffer.GetResourceDataSize(), BUF_Static, CreateInfo);
+			FBufferRHIRef IndexBufferRHI = UE::RHIResourceUtils::CreateIndexBufferFromArray(RHICmdList, TEXT("InvertDistortionMapIndexBuffer"), EBufferUsageFlags::Static, MakeConstArrayView(IndexBuffer));
 
 			RHICmdList.DrawIndexedPrimitive(IndexBufferRHI, 0, 0, NumIndices, 0, NumTriangles, 1);
 		});
diff --git a/Engine/Source/Runtime/Apple/MetalRHI/Private/MetalBindlessDescriptors.cpp b/Engine/Source/Runtime/Apple/MetalRHI/Private/MetalBindlessDescriptors.cpp
--- a/Engine/Source/Runtime/Apple/MetalRHI/Private/MetalBindlessDescriptors.cpp
+++ b/Engine/Source/Runtime/Apple/MetalRHI/Private/MetalBindlessDescriptors.cpp
@@ -41,12 +41,14 @@ FMetalDescriptorHeap::FMetalDescriptorHeap(FMetalDevice& MetalDevice, const ERHI
 void FMetalDescriptorHeap::Init(const int32 HeapSize)
 {
 	FRHICommandListImmediate& RHICmdList = FRHICommandListImmediate::Get();
-	
-	FRHIBufferDesc Desc(HeapSize, 1, BUF_Dynamic | BUF_KeepCPUAccessible | BUF_StructuredBuffer | BUF_UnorderedAccess);
-	FRHIResourceCreateInfo CreateInfo(TEXT("ResourceHeap"));
+
+	const FRHIBufferCreateDesc CreateDesc =
+		FRHIBufferCreateDesc::CreateStructured(TEXT("ResourceHeap"), HeapSize, 1)
+		.AddUsage(EBufferUsageFlags::Dynamic | EBufferUsageFlags::KeepCPUAccessible | EBufferUsageFlags::UnorderedAccess)
+		.DetermineInitialState();
 	
 	ResourceHeapLength = HeapSize;
-	ResourceHeap = new FMetalRHIBuffer(RHICmdList, Device, Desc, CreateInfo);
+	ResourceHeap = new FMetalRHIBuffer(RHICmdList, Device, CreateDesc, nullptr);
 	
 	FMetalRHIBuffer* Buffer = ResourceCast(ResourceHeap.GetReference());
 	
@@ -278,6 +280,11 @@ void FMetalBindlessDescriptorManager::UpdateDescriptorsWithCopy(FMetalRHICommand
 	uint32_t IndexOffset = StandardResources.MinDirtyIndex;
 	uint32_t UpdateSize = ((StandardResources.MaxDirtyIndex - StandardResources.MinDirtyIndex) + 1) * sizeof(IRDescriptorTableEntry);
 	uint32_t UpdateOffset = StandardResources.MinDirtyIndex * sizeof(IRDescriptorTableEntry);
+
+	const FRHIBufferCreateDesc CreateDesc =
+		FRHIBufferCreateDesc::CreateStructured(TEXT("SrcResourceHeap"), UpdateSize, 1)
+		.AddUsage(EBufferUsageFlags::Dynamic | EBufferUsageFlags::KeepCPUAccessible | EBufferUsageFlags::UnorderedAccess)
+		.DetermineInitialState();
 	
 	FMetalBufferPtr SourceBuffer = Device.GetTransferAllocator()->Allocate(UpdateSize);
 	
@@ -328,10 +335,14 @@ void FMetalBindlessDescriptorManager::UpdateDescriptorsWithCompute()
 	FUnorderedAccessViewRHIRef DstDescriptorBufferView;
 	
 	{
-		FRHIResourceCreateInfo CreateInfo(TEXT("DescriptorEntries"), &ComputeDescriptorEntriesCopy);
-		CreateInfo.GPUMask = FRHIGPUMask::GPU0();
-		
-		FBufferRHIRef DescriptorEntriesBuffer = RHICmdList.CreateStructuredBuffer(sizeof(IRDescriptorTableEntry), NumDescriptors * sizeof(IRDescriptorTableEntry), BUF_Dynamic | BUF_ShaderResource | BUF_KeepCPUAccessible, CreateInfo);
+		const FRHIBufferCreateDesc CreateDesc =
+			FRHIBufferCreateDesc::CreateStructured<IRDescriptorTableEntry>(TEXT("DescriptorEntries"), NumDescriptors)
+			.AddUsage(BUF_Dynamic | BUF_ShaderResource | BUF_KeepCPUAccessible)
+			.SetGPUMask(FRHIGPUMask::GPU0())
+			.SetInitActionResourceArray(&ComputeDescriptorEntriesCopy)
+			.DetermineInitialState();
+
+		FBufferRHIRef DescriptorEntriesBuffer = RHICmdList.CreateBuffer(CreateDesc);
 		
 		auto DescriptorEntriesBufferCreateDesc = FRHIViewDesc::CreateBufferSRV()
 			.SetType(FRHIViewDesc::EBufferType::Structured)
@@ -342,10 +353,14 @@ void FMetalBindlessDescriptorManager::UpdateDescriptorsWithCompute()
 	}
 	
 	{
-		FRHIResourceCreateInfo CreateInfo(TEXT("DescriptorIndices"), &ComputeDescriptorIndicesCopy);
-		CreateInfo.GPUMask = FRHIGPUMask::GPU0();
-		
-		FBufferRHIRef DescriptorIndicesBuffer = RHICmdList.CreateStructuredBuffer(sizeof(uint), NumDescriptors*sizeof(uint), BUF_Dynamic | BUF_ShaderResource | BUF_KeepCPUAccessible, CreateInfo);
+		const FRHIBufferCreateDesc CreateDesc =
+		FRHIBufferCreateDesc::CreateStructured<uint>(TEXT("DescriptorIndices"), NumDescriptors)
+			.AddUsage(BUF_Dynamic | BUF_ShaderResource | BUF_KeepCPUAccessible)
+			.SetGPUMask(FRHIGPUMask::GPU0())
+			.SetInitActionResourceArray(&ComputeDescriptorIndicesCopy)
+			.DetermineInitialState();
+
+		FBufferRHIRef DescriptorIndicesBuffer = RHICmdList.CreateBuffer(CreateDesc);
 		
 		auto DescriptorIndicesBufferCreateDesc = FRHIViewDesc::CreateBufferSRV()
 			.SetType(FRHIViewDesc::EBufferType::Structured)
diff --git a/Engine/Source/Runtime/Apple/MetalRHI/Private/MetalDynamicRHI.h b/Engine/Source/Runtime/Apple/MetalRHI/Private/MetalDynamicRHI.h
--- a/Engine/Source/Runtime/Apple/MetalRHI/Private/MetalDynamicRHI.h
+++ b/Engine/Source/Runtime/Apple/MetalRHI/Private/MetalDynamicRHI.h
@@ -97,8 +97,9 @@ public:
 	virtual FGraphicsPipelineStateRHIRef RHICreateGraphicsPipelineState(const FGraphicsPipelineStateInitializer& Initializer) final override;
 	virtual TRefCountPtr<FRHIComputePipelineState> RHICreateComputePipelineState(FRHIComputeShader* ComputeShader) final override;
 	virtual FUniformBufferRHIRef RHICreateUniformBuffer(const void* Contents, const FRHIUniformBufferLayout* Layout, EUniformBufferUsage Usage, EUniformBufferValidation Validation) final override;
-	
-	virtual FBufferRHIRef RHICreateBuffer(FRHICommandListBase& RHICmdList, FRHIBufferDesc const& Desc, ERHIAccess ResourceState, FRHIResourceCreateInfo& CreateInfo) final override;
+
+	[[nodiscard]] virtual FRHIBufferInitializer RHICreateBufferInitializer(FRHICommandListBase& RHICmdList, const FRHIBufferCreateDesc& CreateDesc) final override;
+
 	virtual void RHIReplaceResources(FRHICommandListBase& RHICmdList, TArray<FRHIResourceReplaceInfo>&& ReplaceInfos) final override;
 	
 	virtual void * RHILockBuffer(class FRHICommandListBase& RHICmdList, FRHIBuffer* Buffer, uint32 Offset, uint32 SizeRHI, EResourceLockMode LockMode) final override;
diff --git a/Engine/Source/Runtime/Apple/MetalRHI/Private/MetalRHIBuffer.cpp b/Engine/Source/Runtime/Apple/MetalRHI/Private/MetalRHIBuffer.cpp
--- a/Engine/Source/Runtime/Apple/MetalRHI/Private/MetalRHIBuffer.cpp
+++ b/Engine/Source/Runtime/Apple/MetalRHI/Private/MetalRHIBuffer.cpp
@@ -12,6 +12,7 @@
 #include "Containers/ResourceArray.h"
 #include "RenderUtils.h"
 #include "MetalLLM.h"
+#include "RHICoreBufferInitializer.h"
 #include <objc/runtime.h>
 
 #define METAL_POOL_BUFFER_BACKING 1
@@ -79,27 +80,27 @@ bool FMetalRHIBuffer::UsePrivateMemory() const
 			&& !FMetalCommandQueue::IsUMASystem();
 }
 
-FMetalRHIBuffer::FMetalRHIBuffer(FRHICommandListBase& RHICmdList, FMetalDevice& MetalDevice, FRHIBufferDesc const& InBufferDesc, FRHIResourceCreateInfo& CreateInfo)
-	: FRHIBuffer(InBufferDesc)
+FMetalRHIBuffer::FMetalRHIBuffer(FRHICommandListBase& RHICmdList, FMetalDevice& MetalDevice, const FRHIBufferCreateDesc& CreateDesc, FResourceArrayUploadInterface* InResourceArray)
+	: FRHIBuffer(CreateDesc)
 	, Device(MetalDevice)
-	, Size(InBufferDesc.Size)
+	, Size(CreateDesc.Size)
 	, Mode(BUFFER_STORAGE_MODE)
 {
 #if METAL_RHI_RAYTRACING
-	if (EnumHasAnyFlags(InBufferDesc.Usage, BUF_AccelerationStructure))
+	if (EnumHasAnyFlags(CreateDesc.Usage, BUF_AccelerationStructure))
 	{
 		AccelerationStructureHandle = Device.GetDevice()->newAccelerationStructureWithSize(Size);
 		return;
 	}
 #endif
 
-	const EMetalBufferUsage MetalUsage = GetMetalBufferUsage(InBufferDesc.Usage);
+	const EMetalBufferUsage MetalUsage = GetMetalBufferUsage(CreateDesc.Usage);
 	
-	const bool bIsStatic  	= EnumHasAnyFlags(InBufferDesc.Usage, BUF_Static);
-	const bool bIsDynamic 	= EnumHasAnyFlags(InBufferDesc.Usage, BUF_Dynamic);
-	const bool bIsVolatile	= EnumHasAnyFlags(InBufferDesc.Usage, BUF_Volatile);
-	const bool bIsNull		= EnumHasAnyFlags(InBufferDesc.Usage, BUF_NullResource);
-	const bool bWantsView 	= EnumHasAnyFlags(InBufferDesc.Usage, BUF_ShaderResource | BUF_UnorderedAccess);
+	const bool bIsStatic  	= EnumHasAnyFlags(CreateDesc.Usage, BUF_Static);
+	const bool bIsDynamic 	= EnumHasAnyFlags(CreateDesc.Usage, BUF_Dynamic);
+	const bool bIsVolatile	= EnumHasAnyFlags(CreateDesc.Usage, BUF_Volatile);
+	const bool bIsNull		= EnumHasAnyFlags(CreateDesc.Usage, BUF_NullResource);
+	const bool bWantsView 	= EnumHasAnyFlags(CreateDesc.Usage, BUF_ShaderResource | BUF_UnorderedAccess);
 	
 	uint32_t ValidateTypeCount = (uint32_t)bIsStatic + (uint32_t)bIsDynamic +
 								(uint32_t)bIsVolatile + (uint32_t)bIsNull;
@@ -108,9 +109,9 @@ FMetalRHIBuffer::FMetalRHIBuffer(FRHICommandListBase& RHICmdList, FMetalDevice&
 
 	Mode = UsePrivateMemory() ? MTL::StorageModePrivate : BUFFER_STORAGE_MODE;
 	
-	if (InBufferDesc.Size)
+	if (CreateDesc.Size)
 	{
-		checkf(InBufferDesc.Size <= Device.GetDevice()->maxBufferLength(), TEXT("Requested buffer size larger than supported by device."));
+		checkf(CreateDesc.Size <= Device.GetDevice()->maxBufferLength(), TEXT("Requested buffer size larger than supported by device."));
 		
 		{
 			// Static buffers will never be discarded. You can update them directly.
@@ -129,20 +130,20 @@ FMetalRHIBuffer::FMetalRHIBuffer(FRHICommandListBase& RHICmdList, FMetalDevice&
 #if PLATFORM_MAC
 			// Buffer can be blit encoder copied on lock/unlock, we need to know that the buffer size is large enough for copy operations that are in multiples of
 			// 4 bytes on macOS, iOS can be 1 byte.  Update size to know we have at least this much buffer memory, it will be larger in the end.
-			Size = Align(InBufferDesc.Size, 4);
+			Size = Align(CreateDesc.Size, 4);
 #endif
 			
 			AllocateBuffers();
 		}
 	}
 
-	if (CreateInfo.ResourceArray && InBufferDesc.Size > 0)
+	if (InResourceArray && CreateDesc.Size > 0)
 	{
-		check(InBufferDesc.Size == CreateInfo.ResourceArray->GetResourceDataSize());
+		check(CreateDesc.Size == InResourceArray->GetResourceDataSize());
 
 		if (Data)
 		{
-			FMemory::Memcpy(Data->Data, CreateInfo.ResourceArray->GetResourceData(), InBufferDesc.Size);
+			FMemory::Memcpy(Data->Data, InResourceArray->GetResourceData(), CreateDesc.Size);
 		}
 		else
 		{
@@ -150,17 +151,17 @@ FMetalRHIBuffer::FMetalRHIBuffer(FRHICommandListBase& RHICmdList, FMetalDevice&
 			{
 				if (RHICmdList.IsBottomOfPipe())
 				{
-					void* Backing = this->Lock(true, RLM_WriteOnly, 0, InBufferDesc.Size);
-					FMemory::Memcpy(Backing, CreateInfo.ResourceArray->GetResourceData(), InBufferDesc.Size);
+					void* Backing = this->Lock(true, RLM_WriteOnly, 0, CreateDesc.Size);
+					FMemory::Memcpy(Backing, InResourceArray->GetResourceData(), CreateDesc.Size);
 					this->Unlock(RHICmdList);
 				}
 				else
 				{
-					void* Result = FMemory::Malloc(InBufferDesc.Size, 16);
-					FMemory::Memcpy(Result, CreateInfo.ResourceArray->GetResourceData(), InBufferDesc.Size);
+					void* Result = FMemory::Malloc(CreateDesc.Size, 16);
+					FMemory::Memcpy(Result, InResourceArray->GetResourceData(), CreateDesc.Size);
 
 					RHICmdList.EnqueueLambda(
-						[this, Result, InSize = InBufferDesc.Size](FRHICommandListBase& RHICmdList)
+						[this, Result, InSize = CreateDesc.Size](FRHICommandListBase& RHICmdList)
 						{
 							void* Backing = this->Lock(true, RLM_WriteOnly, 0, InSize);
 							FMemory::Memcpy(Backing, Result, InSize);
@@ -173,7 +174,7 @@ FMetalRHIBuffer::FMetalRHIBuffer(FRHICommandListBase& RHICmdList, FMetalDevice&
 			{
 				FMetalBufferPtr TheBuffer = GetCurrentBuffer();
 				MTL::Buffer* MTLBuffer = TheBuffer->GetMTLBuffer();
-				FMemory::Memcpy(TheBuffer->Contents(), CreateInfo.ResourceArray->GetResourceData(), InBufferDesc.Size);
+				FMemory::Memcpy(TheBuffer->Contents(), InResourceArray->GetResourceData(), CreateDesc.Size);
 #if PLATFORM_MAC 
 				if (Mode == MTL::StorageModeManaged)
 				{
@@ -185,7 +186,7 @@ FMetalRHIBuffer::FMetalRHIBuffer(FRHICommandListBase& RHICmdList, FMetalDevice&
 		}
 
 		// Discard the resource array's contents.
-		CreateInfo.ResourceArray->Discard();
+		InResourceArray->Discard();
 	}
 }
 
@@ -592,22 +593,50 @@ void FMetalRHIBuffer::ReleaseOwnership()
 	FRHIBuffer::ReleaseOwnership();
 }
 
-FBufferRHIRef FMetalDynamicRHI::RHICreateBuffer(FRHICommandListBase& RHICmdList, FRHIBufferDesc const& Desc, ERHIAccess /*ResourceState*/, FRHIResourceCreateInfo& CreateInfo)
+static FRHIBufferCreateDesc MetalModifyBufferCreateDesc(const FRHIBufferCreateDesc& InCreateDesc)
 {
-    MTL_SCOPED_AUTORELEASE_POOL;
+	FRHIBufferCreateDesc CreateDesc(InCreateDesc);
 
-    // No life-time usage information? Enforce Dynamic.
-    if (!EnumHasAnyFlags(Desc.Usage, BUF_Static | BUF_Dynamic | BUF_Volatile | BUF_NullResource))
-    {
-        FRHIBufferDesc Copy = Desc;
-        Copy.Usage |= BUF_Dynamic;
+	// No life-time usage information? Enforce Dynamic.
+	if (!EnumHasAnyFlags(CreateDesc.Usage, EBufferUsageFlags::Static | EBufferUsageFlags::Dynamic | EBufferUsageFlags::Volatile | EBufferUsageFlags::NullResource))
+	{
+		CreateDesc.AddUsage(EBufferUsageFlags::Dynamic);
+	}
 
-        return new FMetalRHIBuffer(RHICmdList, *Device, Copy, CreateInfo);
-    }
-    else
-    {
-        return new FMetalRHIBuffer(RHICmdList, *Device, Desc, CreateInfo);
-    }
+	return CreateDesc;
+}
+
+FRHIBufferInitializer FMetalDynamicRHI::RHICreateBufferInitializer(FRHICommandListBase& RHICmdList, const FRHIBufferCreateDesc& InCreateDesc)
+{
+	MTL_SCOPED_AUTORELEASE_POOL;
+
+	const FRHIBufferCreateDesc CreateDesc = MetalModifyBufferCreateDesc(InCreateDesc);
+
+	// TODO: Can Metal use UE::RHICore::CreateUnifiedMemoryBufferInitializer ?
+
+	FMetalRHIBuffer* Buffer = new FMetalRHIBuffer(RHICmdList, *Device, CreateDesc, CreateDesc.InitialData);
+
+	if (CreateDesc.IsNull() || CreateDesc.InitAction == ERHIBufferInitAction::ResourceArray || CreateDesc.InitAction == ERHIBufferInitAction::Default)
+	{
+		return UE::RHICore::FDefaultBufferInitializer(RHICmdList, Buffer);
+	}
+
+	if (CreateDesc.InitAction == ERHIBufferInitAction::Zeroed)
+	{
+		void* WritableData = RHICmdList.LockBuffer(Buffer, 0, CreateDesc.Size, RLM_WriteOnly);
+		FMemory::Memzero(WritableData, CreateDesc.Size);
+		RHICmdList.UnlockBuffer(Buffer);
+
+		return UE::RHICore::FDefaultBufferInitializer(RHICmdList, Buffer);
+	}
+
+	if (CreateDesc.InitAction == ERHIBufferInitAction::Initializer)
+	{
+		// Use LockBuffer + UnlockBuffer to allow the caller to write initial buffer data
+		return UE::RHICore::FLockBufferInitializer(RHICmdList, Buffer);
+	}
+
+	return UE::RHICore::HandleUnknownBufferInitializerInitAction(RHICmdList, CreateDesc);
 }
 
 struct FMetalRHILockData
diff --git a/Engine/Source/Runtime/Apple/MetalRHI/Private/MetalRayTracing.cpp b/Engine/Source/Runtime/Apple/MetalRHI/Private/MetalRayTracing.cpp
--- a/Engine/Source/Runtime/Apple/MetalRHI/Private/MetalRayTracing.cpp
+++ b/Engine/Source/Runtime/Apple/MetalRHI/Private/MetalRayTracing.cpp
@@ -360,10 +360,14 @@ FMetalRayTracingGeometry::FMetalRayTracingGeometry(FRHICommandListBase& RHICmdLi
 	}
 
 	FString DebugNameString = Initializer.DebugName.ToString();
-	FRHIResourceCreateInfo BlasBufferCreateInfo(*DebugNameString);
+
+	const FRHIBufferCreateDesc CreateDesc =
+		FRHIBufferCreateDesc::Create(*DebugNameString, SizeInfo.ResultSize, 0, EBufferUsageFlags::AccelerationStructure)
+		.SetInitialState(ERHIAccess::BVHWrite);
+
 	for (uint32 i = 0; i < MaxNumAccelerationStructure; i++)
 	{
-		AccelerationStructure[i] = ResourceCast(RHICmdList.CreateBuffer(SizeInfo.ResultSize, BUF_AccelerationStructure, 0, ERHIAccess::BVHWrite, BlasBufferCreateInfo).GetReference());
+		AccelerationStructure[i] = ResourceCast(RHICmdList.CreateBuffer(CreateDesc).GetReference());
 		check(AccelerationStructure[i]);
 
 		AccelerationStructure[i]->AccelerationStructureHandle.SetLabel(TCHAR_TO_ANSI(*DebugNameString));
@@ -430,8 +434,8 @@ FMetalRayTracingScene::FMetalRayTracingScene(FRayTracingSceneInitializer InIniti
 
 	MutableAccelerationStructures = [[NSMutableArray<id<MTLAccelerationStructure>> new] init];
 
-	const uint32 ParameterBufferSize = FMath::Max<uint32>(1, Initializer.NumTotalSegments) * sizeof(FHitGroupSystemRootConstants);
-	FRHIResourceCreateInfo ParameterBufferCreateInfo(TEXT("RayTracingSceneMetadata"));
+	//const uint32 ParameterBufferSize = FMath::Max<uint32>(1, Initializer.NumTotalSegments) * sizeof(FHitGroupSystemRootConstants);
+	//FRHIResourceCreateInfo ParameterBufferCreateInfo(TEXT("RayTracingSceneMetadata"));
 	//PerInstanceGeometryParameterBuffer = ResourceCast(FRHICommandListExecutor::GetImmediateCommandList().CreateBuffer(ParameterBufferSize,  BUF_StructuredBuffer | BUF_ShaderResource, sizeof(FHitGroupSystemRootConstants), ERHIAccess::SRVMask, ParameterBufferCreateInfo).GetReference());
 
 	//FShaderResourceViewInitializer ViewInitializer(PerInstanceGeometryParameterBuffer);
@@ -563,8 +567,12 @@ void FMetalRayTracingScene::BuildAccelerationStructure(
 	TRefCountPtr<FMetalRHIBuffer> ScratchBuffer;
 	if (InScratchBuffer == nullptr)
 	{
-		FRHIResourceCreateInfo ScratchBufferCreateInfo(TEXT("BuildScratchTLAS"));
-		ScratchBuffer = ResourceCast(RHICreateBuffer(SizeInfo.BuildScratchSize, BUF_UnorderedAccess | BUF_StructuredBuffer, 0, ERHIAccess::UAVCompute, ScratchBufferCreateInfo).GetReference());
+		const FRHIBufferCreateDesc CreateDesc =
+			FRHIBufferCreateDesc::CreateStructured(TEXT("BuildScratchTLAS"), SizeInfo.BuildScratchSize, 0)
+			.AddUsage(EBufferUsageFlags::UnorderedAccess)
+			.SetInitialState(ERHIAccess::UAVCompute);
+
+		ScratchBuffer = ResourceCast(RHICreateBuffer(CreateDesc).GetReference());
 		InScratchBuffer = ScratchBuffer.GetReference();
 		ScratchOffset = 0;
 	}
diff --git a/Engine/Source/Runtime/Apple/MetalRHI/Public/MetalResources.h b/Engine/Source/Runtime/Apple/MetalRHI/Public/MetalResources.h
--- a/Engine/Source/Runtime/Apple/MetalRHI/Public/MetalResources.h
+++ b/Engine/Source/Runtime/Apple/MetalRHI/Public/MetalResources.h
@@ -293,7 +293,7 @@ public:
 	// Matches other RHIs
 	static constexpr const uint32 MetalMaxNumBufferedFrames = 4;
 	
-	FMetalRHIBuffer(FRHICommandListBase& RHICmdList, FMetalDevice& MetalDevice, FRHIBufferDesc const& InBufferDesc, FRHIResourceCreateInfo& CreateInfo);
+	FMetalRHIBuffer(FRHICommandListBase& RHICmdList, FMetalDevice& MetalDevice, const FRHIBufferCreateDesc& CreateDesc, FResourceArrayUploadInterface* InResourceArray);
 	virtual ~FMetalRHIBuffer();
 	
 	bool RequiresTransferBuffer();
diff --git a/Engine/Source/Runtime/Core/Public/Containers/ResourceArray.h b/Engine/Source/Runtime/Core/Public/Containers/ResourceArray.h
--- a/Engine/Source/Runtime/Core/Public/Containers/ResourceArray.h
+++ b/Engine/Source/Runtime/Core/Public/Containers/ResourceArray.h
@@ -16,6 +16,12 @@ struct FResourceArrayUploadInterface
 	/** Returns size of resource data allocation */
 	[[nodiscard]] virtual uint32 GetResourceDataSize() const = 0;
 
+	template<typename TElement>
+	TConstArrayView<TElement> GetResourceDataView() const
+	{
+		return TConstArrayView<TElement>(reinterpret_cast<const TElement*>(GetResourceData()), GetResourceDataSize());
+	}
+
 	/** Called on non-UMA systems after the RHI has copied the resource data, and no longer needs the CPU's copy. */
 	virtual void Discard() = 0;
 };
diff --git a/Engine/Source/Runtime/D3D12RHI/Private/D3D12Adapter.h b/Engine/Source/Runtime/D3D12RHI/Private/D3D12Adapter.h
--- a/Engine/Source/Runtime/D3D12RHI/Private/D3D12Adapter.h
+++ b/Engine/Source/Runtime/D3D12RHI/Private/D3D12Adapter.h
@@ -379,16 +379,13 @@ public:
 
 	FD3D12Buffer* CreateRHIBuffer(
 		const D3D12_RESOURCE_DESC& Desc,
-		uint32 Alignment, FRHIBufferDesc const& BufferDesc,
+		uint32 Alignment,
+		const FRHIBufferCreateDesc& CreateDesc,
 		ED3D12ResourceStateMode InResourceStateMode,
 		D3D12_RESOURCE_STATES InCreateState,
-		bool bHasInitialData,
-		const FRHIGPUMask& InGPUMask,
-		ID3D12ResourceAllocator* ResourceAllocator,
-		const TCHAR* InDebugName,
-		ERHIAccess InRHICreationState,
-		const FName& InOwnerName = NAME_None,
-		const FName& InClassName = NAME_None);
+		bool bKeepUnlocked = false,
+		ID3D12ResourceAllocator* ResourceAllocator = nullptr
+	);
 
 	void CreateUAVAliasResource(
 		D3D12_CLEAR_VALUE* ClearValuePtr,
diff --git a/Engine/Source/Runtime/D3D12RHI/Private/D3D12Buffer.cpp b/Engine/Source/Runtime/D3D12RHI/Private/D3D12Buffer.cpp
--- a/Engine/Source/Runtime/D3D12RHI/Private/D3D12Buffer.cpp
+++ b/Engine/Source/Runtime/D3D12RHI/Private/D3D12Buffer.cpp
@@ -6,9 +6,10 @@ D3D12Buffer.cpp: D3D Common code for buffers.
 
 #include "D3D12RHIPrivate.h"
 #include "ProfilingDebugging/AssetMetadataTrace.h"
+#include "RHICoreBufferInitializer.h"
 #include "RHICoreStats.h"
 
-FName GetRHIBufferClassName(const FName& ClassName)
+static FName GetRHIBufferClassName(const FName ClassName)
 {
 	const static FLazyName RHIBufferName(TEXT("FRHIBuffer"));
 	return (ClassName == NAME_None) ? RHIBufferName : ClassName;
@@ -69,66 +70,23 @@ void FD3D12Buffer::UploadResourceData(FD3D12CommandContext& CommandContext, D3D1
 	DestinationResourceLocation.UnlockPoolData();
 }
 
-void FD3D12Buffer::UploadResourceData(FRHICommandListBase& InRHICmdList, FRHIGPUMask GPUMask, D3D12_RESOURCE_STATES InDestinationState, const void* SourceData, int32 SourceDataSize)
+static FD3D12ResourceLocation AllocateUploadMemory(FD3D12Device* Device, uint32 Size, uint32 Alignment)
 {
-	FD3D12Device* ParentDevice = GetParentDevice();
-
-	const uint32 BufferSize = GetSize();
-	check(BufferSize == SourceDataSize);
-
-	const bool bOnAsyncThread = !IsInRHIThread() && !IsInRenderingThread();
-
-	// Get an upload heap and initialize data
-	FD3D12ResourceLocation SrcResourceLoc(ParentDevice);
-	void* pData;
-	if (bOnAsyncThread)
+	FD3D12ResourceLocation Location(Device);
+	if (!IsInRHIThread() && !IsInRenderingThread())
 	{
-		pData = ParentDevice->GetParentAdapter()->GetUploadHeapAllocator(ParentDevice->GetGPUIndex()).AllocUploadResource(BufferSize, 4u, SrcResourceLoc);
+		Device->GetParentAdapter()->GetUploadHeapAllocator(Device->GetGPUIndex()).AllocUploadResource(Size, Alignment, Location);
 	}
 	else
 	{
-		pData = ParentDevice->GetDefaultFastAllocator().Allocate(BufferSize, 4UL, &SrcResourceLoc);
+		Device->GetDefaultFastAllocator().Allocate(Size, Alignment, &Location);
 	}
-	check(pData);
-	FMemory::Memcpy(pData, SourceData, BufferSize);
-
-	InRHICmdList.EnqueueLambda(
-		[this, GPUMask, InDestinationState, SrcResourceLoc = MoveTemp(SrcResourceLoc)](FRHICommandListBase& ExecutingCmdList)
-		{
-			const FRHIGPUMask EffectiveMask = GPUMask & ExecutingCmdList.GetGPUMask();
-			for (uint32 GPUIndex : EffectiveMask)
-			{
-				FD3D12CommandContext& CommandContext = FD3D12CommandContext::Get(ExecutingCmdList, GPUIndex);
-				FD3D12Buffer* DeviceBuffer = CommandContext.RetrieveObject<FD3D12Buffer>(this, GPUIndex);
-
-				UploadResourceData(CommandContext, InDestinationState, DeviceBuffer->ResourceLocation, SrcResourceLoc, GetSize());
-			}
-		});
+	return Location;
 }
 
-void FD3D12Buffer::UploadResourceData(FRHICommandListBase& RHICmdList, FResourceArrayUploadInterface* InResourceArray, D3D12_RESOURCE_STATES InDestinationState, const TCHAR* AssetName, const FName& ClassName, const FName& PackageName)
+static FD3D12ResourceLocation AllocateUploadMemory(FD3D12Buffer* Buffer, const FRHIBufferDesc& Desc)
 {
-	UE_TRACE_METADATA_SCOPE_ASSET_FNAME(AssetName, ClassName, PackageName);
-
-	check(InResourceArray);
-	check(ResourceLocation.IsValid());
-
-	uint32 BufferSize = GetSize();
-	check(BufferSize == InResourceArray->GetResourceDataSize());
-
-	if (EnumHasAnyFlags(GetUsage(), BUF_AnyDynamic))
-	{
-		// Copy directly in mapped data
-		void* MappedUploadData = ResourceLocation.GetMappedBaseAddress(); 
-		FMemory::Memcpy(MappedUploadData, InResourceArray->GetResourceData(), BufferSize);
-	}
-	else
-	{
-		UploadResourceData(RHICmdList, FRHIGPUMask::All(), InDestinationState, InResourceArray->GetResourceData(), BufferSize);
-	}
-
-	// Discard the resource array's contents.
-	InResourceArray->Discard();
+	return AllocateUploadMemory(Buffer->GetParentDevice(), Desc.Size, Buffer->BufferAlignment);
 }
 
 void FD3D12CommandContext::CopyBufferRegionChecked(
@@ -156,7 +114,7 @@ FD3D12SyncPointRef FD3D12Buffer::UploadResourceDataViaCopyQueue(FD3D12CommandCon
 
 	// assume not dynamic and not on async thread (probably fine but untested)
 	check(IsInRHIThread() || IsInRenderingThread());
-	check(!(GetUsage() & BUF_AnyDynamic));
+	check(!(GetUsage() & EBufferUsageFlags::AnyDynamic));
 
 	uint32 BufferSize = GetSize();
 
@@ -226,7 +184,7 @@ void FD3D12Adapter::AllocateBuffer(FD3D12Device* Device,
 		);
 	}
 	
-	if (EnumHasAnyFlags(InUsage, BUF_AnyDynamic))
+	if (EnumHasAnyFlags(InUsage, EBufferUsageFlags::AnyDynamic))
 	{
 		check(ResourceAllocator == nullptr);
 		check(InResourceStateMode != ED3D12ResourceStateMode::MultiState);
@@ -252,86 +210,57 @@ void FD3D12Adapter::AllocateBuffer(FD3D12Device* Device,
 FD3D12Buffer* FD3D12Adapter::CreateRHIBuffer(
 	const D3D12_RESOURCE_DESC& InDesc,
 	uint32 Alignment,
-	FRHIBufferDesc const& BufferDesc,
+	const FRHIBufferCreateDesc& CreateDesc,
 	ED3D12ResourceStateMode InResourceStateMode,
 	D3D12_RESOURCE_STATES InCreateState,
 	bool bKeepUnlocked,
-	const FRHIGPUMask& InGPUMask,
-	ID3D12ResourceAllocator* ResourceAllocator,
-	const TCHAR* InDebugName,
-	ERHIAccess InRHICreationState,
-	const FName& OwnerName,
-	const FName& ClassName)
+	ID3D12ResourceAllocator* ResourceAllocator)
 {
 	TRACE_CPUPROFILER_EVENT_SCOPE(D3D12RHI::CreateRHIBuffer);
 	SCOPE_CYCLE_COUNTER(STAT_D3D12CreateBufferTime);
 
-	UE_TRACE_METADATA_SCOPE_ASSET_FNAME(FName(InDebugName), GetRHIBufferClassName(ClassName), OwnerName);
+	UE_TRACE_METADATA_SCOPE_ASSET_FNAME(FName(CreateDesc.DebugName), GetRHIBufferClassName(CreateDesc.ClassName), CreateDesc.OwnerName);
 
-	check(InDesc.Width >= BufferDesc.Size);
+	check(InDesc.Width >= CreateDesc.Size);
 
+	const bool bDynamic = EnumHasAnyFlags(CreateDesc.Usage, EBufferUsageFlags::AnyDynamic);
 	FD3D12Buffer* BufferOut = nullptr;
 
-	if (EnumHasAnyFlags(BufferDesc.Usage, BUF_AnyDynamic))
-	{
-		const uint32 FirstGPUIndex = InGPUMask.GetFirstIndex();
-
-		FD3D12Buffer* NewBuffer0 = nullptr;
-		BufferOut = CreateLinkedObject<FD3D12Buffer>(InGPUMask, [&](FD3D12Device* Device, FD3D12Buffer* FirstLinkedObject)
-		{
-			FD3D12Buffer* NewBuffer = new FD3D12Buffer(Device, BufferDesc);
-			NewBuffer->SetTrackedAccess_Unsafe(InRHICreationState);
-			NewBuffer->BufferAlignment = Alignment;
-
-#if NAME_OBJECTS
-			if (InDebugName)
-			{
-				TRACE_CPUPROFILER_EVENT_SCOPE(D3D12RHI::SetDebugName);
-				NewBuffer->SetName(InDebugName);
-			}
-#endif // NAME_OBJECTS
+	// Theoretically, we could assert if GPUMask isn't correct, but at the moment the RDG and RHI buffer descriptions don't include the
+	// GPU mask, so there's no way for the caller to configure it (only the lower level CreateInfo includes it).  Note that differentiation
+	// for NNE (DirectML) is required beyond just setting the mask anyway, in the sense of forcing separate GPU0 visible only heaps, not just
+	// filtering which GPU copies are allocated.  Because this is necessary to solve a crash, it's higher priority than GPUMask support,
+	// which may be added in the future.
+	FRHIGPUMask GPUMask = EnumHasAnyFlags(CreateDesc.Usage, EBufferUsageFlags::NNE) ? FRHIGPUMask::GPU0() : CreateDesc.GPUMask;
 
-			if ((Device->GetGPUIndex() == FirstGPUIndex) || EnumHasAnyFlags(BufferDesc.Usage, BUF_MultiGPUAllocate))
-			{
-				AllocateBuffer(Device, InDesc, BufferDesc.Size, BufferDesc.Usage, InResourceStateMode, InCreateState, Alignment, NewBuffer, NewBuffer->ResourceLocation, ResourceAllocator, InDebugName, OwnerName);
-				NewBuffer0 = NewBuffer;
-			}
-			else
-			{
-				check(NewBuffer0);
-				FD3D12ResourceLocation::ReferenceNode(Device, NewBuffer->ResourceLocation, NewBuffer0->ResourceLocation);
-			}
+	const uint32 FirstGPUIndex = GPUMask.GetFirstIndex();
 
-			return NewBuffer;
-		});
-	}
-	else
+	FD3D12Buffer* NewBuffer0 = nullptr;
+	BufferOut = CreateLinkedObject<FD3D12Buffer>(GPUMask, [&](FD3D12Device* Device, FD3D12Buffer* FirstLinkedObject)
 	{
-		BufferOut = CreateLinkedObject<FD3D12Buffer>(InGPUMask, [&](FD3D12Device* Device, FD3D12Buffer* FirstLinkedObject)
-		{
-			FD3D12Buffer* NewBuffer = new FD3D12Buffer(Device, BufferDesc);
-			NewBuffer->SetTrackedAccess_Unsafe(InRHICreationState);
-			NewBuffer->BufferAlignment = Alignment;
+		FD3D12Buffer* NewBuffer = new FD3D12Buffer(Device, CreateDesc);
+		NewBuffer->BufferAlignment = Alignment;
 
-#if NAME_OBJECTS
-			if (InDebugName)
-			{
-				TRACE_CPUPROFILER_EVENT_SCOPE(D3D12RHI::SetDebugName);
-				NewBuffer->SetName(InDebugName);
-			}
-#endif // NAME_OBJECTS
+		if (!bDynamic || (Device->GetGPUIndex() == FirstGPUIndex) || EnumHasAnyFlags(CreateDesc.Usage, EBufferUsageFlags::MultiGPUAllocate))
+		{
+			AllocateBuffer(Device, InDesc, CreateDesc.Size, CreateDesc.Usage, InResourceStateMode, InCreateState, Alignment, NewBuffer, NewBuffer->ResourceLocation, ResourceAllocator, CreateDesc.DebugName, CreateDesc.OwnerName);
+			NewBuffer0 = NewBuffer;
+		}
+		else
+		{
+			check(NewBuffer0);
+			FD3D12ResourceLocation::ReferenceNode(Device, NewBuffer->ResourceLocation, NewBuffer0->ResourceLocation);
+		}
 
-			AllocateBuffer(Device, InDesc, BufferDesc.Size, BufferDesc.Usage, InResourceStateMode, InCreateState, Alignment, NewBuffer, NewBuffer->ResourceLocation, ResourceAllocator, InDebugName, OwnerName);
-			
-			// Unlock immediately if there is no initial data
-			if (!bKeepUnlocked)
-			{
-				NewBuffer->ResourceLocation.UnlockPoolData();
-			}
+		
+		// Unlock immediately if there is no initial data
+		if (!bDynamic && !bKeepUnlocked)
+		{
+			NewBuffer->ResourceLocation.UnlockPoolData();
+		}
 
-			return NewBuffer;
-		});
-	}
+		return NewBuffer;
+	});
 
 	// Don't track transient buffer stats here
 	if (!BufferOut->ResourceLocation.IsTransient())
@@ -350,15 +279,15 @@ void FD3D12Buffer::Rename(FD3D12ContextArray const& Contexts, FD3D12ResourceLoca
 
 void FD3D12Buffer::RenameLDAChain(FD3D12ContextArray const& Contexts, FD3D12ResourceLocation& NewLocation)
 {
-	// Dynamic buffers use cross-node resources (with the exception of BUF_MultiGPUAllocate)
-	//ensure(GetUsage() & BUF_AnyDynamic);
+	// Dynamic buffers use cross-node resources (with the exception of EBufferUsageFlags::MultiGPUAllocate)
+	//ensure(GetUsage() & EBufferUsageFlags::AnyDynamic);
 	Rename(Contexts, NewLocation);
 
 	if (GNumExplicitGPUsForRendering > 1)
 	{
 		ensure(GetParentDevice() == NewLocation.GetParentDevice());
 
-		if (EnumHasAnyFlags(GetUsage(), BUF_MultiGPUAllocate) == false)
+		if (EnumHasAnyFlags(GetUsage(), EBufferUsageFlags::MultiGPUAllocate) == false)
 		{
 			ensure(IsHeadLink());
 
@@ -399,126 +328,193 @@ void FD3D12Buffer::ReleaseOwnership()
 	ResourceLocation.Clear();
 }
 
-void FD3D12Buffer::GetResourceDescAndAlignment(uint64 InSize, uint32 InStride, EBufferUsageFlags InUsage, D3D12_RESOURCE_DESC& ResourceDesc, uint32& Alignment)
+void FD3D12Buffer::GetResourceDescAndAlignment(const FRHIBufferCreateDesc& CreateDesc, D3D12_RESOURCE_DESC& ResourceDesc, uint32& Alignment)
 {
-	ResourceDesc = CD3DX12_RESOURCE_DESC::Buffer(InSize);
+	ResourceDesc = CD3DX12_RESOURCE_DESC::Buffer(CreateDesc.Size);
 
 	// Align size to 16 so RAW buffer view can be created without loosing any data at the end when dividing num elements by 4
 	ResourceDesc.Width = Align(ResourceDesc.Width, RHI_RAW_VIEW_ALIGNMENT);
 
-	if (EnumHasAnyFlags(InUsage, BUF_UnorderedAccess))
+	if (EnumHasAnyFlags(CreateDesc.Usage, EBufferUsageFlags::UnorderedAccess))
 	{
 		ResourceDesc.Flags |= D3D12_RESOURCE_FLAG_ALLOW_UNORDERED_ACCESS;
 	}
 
-	if (!EnumHasAnyFlags(InUsage, BUF_ShaderResource | BUF_AccelerationStructure))
+	if (!EnumHasAnyFlags(CreateDesc.Usage, EBufferUsageFlags::ShaderResource | EBufferUsageFlags::AccelerationStructure))
 	{
 		ResourceDesc.Flags |= D3D12_RESOURCE_FLAG_DENY_SHADER_RESOURCE;
 	}
 
-	if (EnumHasAnyFlags(InUsage, BUF_DrawIndirect))
+	if (EnumHasAnyFlags(CreateDesc.Usage, EBufferUsageFlags::DrawIndirect))
 	{
 		ResourceDesc.Flags |= D3D12RHI_RESOURCE_FLAG_ALLOW_INDIRECT_BUFFER;
 	}
 
-	if (EnumHasAnyFlags(InUsage, BUF_Shared))
+	if (EnumHasAnyFlags(CreateDesc.Usage, EBufferUsageFlags::Shared))
 	{
 		ResourceDesc.Flags |= D3D12_RESOURCE_FLAG_ALLOW_SIMULTANEOUS_ACCESS;
 	}
 
-	if (EnumHasAnyFlags(InUsage, BUF_ReservedResource))
+	if (EnumHasAnyFlags(CreateDesc.Usage, EBufferUsageFlags::ReservedResource))
 	{
-		checkf(InStride <= GRHIGlobals.ReservedResources.TileSizeInBytes,
-			TEXT("Reserved buffer stride %d must not be greater than the reserved resource tile size %d"), 
-			InStride, GRHIGlobals.ReservedResources.TileSizeInBytes);
+		checkf(CreateDesc.Stride <= GRHIGlobals.ReservedResources.TileSizeInBytes,
+			   TEXT("Reserved buffer stride %d must not be greater than the reserved resource tile size %d"),
+			   CreateDesc.Stride, GRHIGlobals.ReservedResources.TileSizeInBytes);
 
 		Alignment = GRHIGlobals.ReservedResources.TileSizeInBytes;
 	}
 	else
 	{
 		// Structured buffers, non-ByteAddress buffers, need to be aligned to their stride to ensure that they can be addressed correctly with element based offsets.
-		Alignment = (InStride > 0) && (EnumHasAnyFlags(InUsage, BUF_StructuredBuffer) || !EnumHasAnyFlags(InUsage, BUF_ByteAddressBuffer | BUF_DrawIndirect)) ? FMath::LeastCommonMultiplier(InStride, RHI_RAW_VIEW_ALIGNMENT) : RHI_RAW_VIEW_ALIGNMENT;
-	}
-}
-
-FBufferRHIRef FD3D12DynamicRHI::RHICreateBuffer(FRHICommandListBase& RHICmdList, FRHIBufferDesc const& Desc, ERHIAccess ResourceState, FRHIResourceCreateInfo& CreateInfo)
-{
-	return CreateBuffer(RHICmdList, Desc, ResourceState, CreateInfo);
-}
-
-FBufferRHIRef FD3D12DynamicRHI::CreateBuffer(FRHICommandListBase& RHICmdList, FRHIBufferDesc const& BufferDesc, ERHIAccess InResourceState, FRHIResourceCreateInfo& CreateInfo)
-{
-	if (BufferDesc.IsNull())
-	{
-		return GetAdapter().CreateLinkedObject<FD3D12Buffer>(CreateInfo.GPUMask, [BufferDesc](FD3D12Device* Device, FD3D12Buffer* FirstLinkedObject)
-		{
-			return new FD3D12Buffer(Device, BufferDesc);
-		});
+		Alignment = (CreateDesc.Stride > 0) && (EnumHasAnyFlags(CreateDesc.Usage, EBufferUsageFlags::StructuredBuffer) || !EnumHasAnyFlags(CreateDesc.Usage, EBufferUsageFlags::ByteAddressBuffer | EBufferUsageFlags::DrawIndirect)) ? FMath::LeastCommonMultiplier(CreateDesc.Stride, RHI_RAW_VIEW_ALIGNMENT) : RHI_RAW_VIEW_ALIGNMENT;
 	}
-
-	return CreateD3D12Buffer(&RHICmdList, BufferDesc, InResourceState, CreateInfo);
 }
 
-FD3D12Buffer* FD3D12DynamicRHI::CreateD3D12Buffer(class FRHICommandListBase* RHICmdList, FRHIBufferDesc const& BufferDesc, ERHIAccess InResourceState, FRHIResourceCreateInfo& CreateInfo, ID3D12ResourceAllocator* ResourceAllocator, bool bForceKeepUnlocked)
+FD3D12DynamicRHI::FCreateBufferInternalResult FD3D12DynamicRHI::CreateBufferInternal(const FRHIBufferCreateDesc& CreateDesc, bool bHasInitialData, ID3D12ResourceAllocator* ResourceAllocator)
 {
-	FName TraceClassName = GetRHIBufferClassName(CreateInfo.GetTraceClassName());
+	D3D12_RESOURCE_DESC Desc{};
+	uint32 Alignment{};
+	FD3D12Buffer::GetResourceDescAndAlignment(CreateDesc, Desc, Alignment);
 
-	D3D12_RESOURCE_DESC Desc;
-	uint32 Alignment;
-	FD3D12Buffer::GetResourceDescAndAlignment(BufferDesc.Size, BufferDesc.Stride, BufferDesc.Usage, Desc, Alignment);
-
-	ED3D12ResourceStateMode StateMode = EnumHasAllFlags(BufferDesc.Usage, BUF_AccelerationStructure)
-		? ED3D12ResourceStateMode::SingleState 
+	ED3D12ResourceStateMode StateMode = EnumHasAllFlags(CreateDesc.Usage, EBufferUsageFlags::AccelerationStructure)
+		? ED3D12ResourceStateMode::SingleState
 		: ED3D12ResourceStateMode::Default;
 
-	const bool bHasInitialData = CreateInfo.ResourceArray != nullptr;
-
-	const bool bIsDynamic = EnumHasAnyFlags(BufferDesc.Usage, BUF_AnyDynamic);
+	const bool bIsDynamic = EnumHasAnyFlags(CreateDesc.Usage, EBufferUsageFlags::AnyDynamic);
 
-	if (EnumHasAnyFlags(BufferDesc.Usage, BUF_ReservedResource))
+	if (EnumHasAnyFlags(CreateDesc.Usage, EBufferUsageFlags::ReservedResource))
 	{
 		checkf(!bHasInitialData, TEXT("Reserved resources may not have initial data"));
 		checkf(!bIsDynamic, TEXT("Reserved resources may not be dynamic"));
 		checkf(!ResourceAllocator, TEXT("Reserved resources may not use a custom resource allocator"));
 	}
 
-	D3D12_HEAP_TYPE HeapType = bIsDynamic ? D3D12_HEAP_TYPE_UPLOAD : D3D12_HEAP_TYPE_DEFAULT;
+	const D3D12_HEAP_TYPE HeapType = bIsDynamic ? D3D12_HEAP_TYPE_UPLOAD : D3D12_HEAP_TYPE_DEFAULT;
 	const FD3D12Resource::FD3D12ResourceTypeHelper Type(Desc, HeapType);
 
 	// Does this resource support tracking?
 	const bool bSupportResourceStateTracking = !bIsDynamic && FD3D12DefaultBufferAllocator::IsPlacedResource(Desc.Flags, StateMode, Alignment) && Type.bWritable;
 
 	// Initial state is derived from the InResourceState if it supports tracking
-	D3D12_RESOURCE_STATES DesiredState = bSupportResourceStateTracking ? Type.GetOptimalInitialState(InResourceState, false) :
-		FD3D12DefaultBufferAllocator::GetDefaultInitialResourceState(HeapType, BufferDesc.Usage, StateMode);
+	const D3D12_RESOURCE_STATES DesiredState = bSupportResourceStateTracking ? Type.GetOptimalInitialState(CreateDesc.InitialState, false) :
+		FD3D12DefaultBufferAllocator::GetDefaultInitialResourceState(HeapType, CreateDesc.Usage, StateMode);
 
 	// Setup the state at which the resource needs to be created - copy dest only supported for placed resources
-	D3D12_RESOURCE_STATES CreateState = (CreateInfo.ResourceArray && bSupportResourceStateTracking) ? D3D12_RESOURCE_STATE_COPY_DEST : DesiredState;
+	const D3D12_RESOURCE_STATES CreateState = (bHasInitialData && bSupportResourceStateTracking) ? D3D12_RESOURCE_STATE_COPY_DEST : DesiredState;
 
-	// Theoretically, we could assert if GPUMask isn't correct, but at the moment the RDG and RHI buffer descriptions don't include the
-	// GPU mask, so there's no way for the caller to configure it (only the lower level CreateInfo includes it).  Note that differentiation
-	// for NNE (DirectML) is required beyond just setting the mask anyway, in the sense of forcing separate GPU0 visible only heaps, not just
-	// filtering which GPU copies are allocated.  Because this is necessary to solve a crash, it's higher priority than GPUMask support,
-	// which may be added in the future.
-	FRHIGPUMask GPUMask = EnumHasAnyFlags(BufferDesc.Usage, EBufferUsageFlags::NNE) ? FRHIGPUMask::GPU0() : CreateInfo.GPUMask;
-
-	FD3D12Buffer* Buffer = GetAdapter().CreateRHIBuffer(Desc, Alignment, BufferDesc, StateMode, CreateState, bHasInitialData || bForceKeepUnlocked, GPUMask, ResourceAllocator, CreateInfo.DebugName, InResourceState, CreateInfo.OwnerName, TraceClassName);
+	FD3D12Buffer* Buffer = GetAdapter().CreateRHIBuffer(Desc, Alignment, CreateDesc, StateMode, CreateState, bHasInitialData, ResourceAllocator);
 	check(Buffer->ResourceLocation.IsValid());
 
-	// Copy the resource data if available 
-	if (bHasInitialData)
+	return FCreateBufferInternalResult{ Buffer, DesiredState };
+}
+
+static FRHIBufferInitializer CreateBufferInitializerForWriting(FRHICommandListBase& RHICmdList, const FD3D12DynamicRHI::FCreateBufferInternalResult& CreateResult, const FRHIBufferCreateDesc& CreateDesc)
+{
+	FD3D12Buffer* Buffer = CreateResult.Buffer;
+	D3D12_RESOURCE_STATES DesiredState = CreateResult.DesiredState;
+
+	if (EnumHasAnyFlags(CreateDesc.Usage, EBufferUsageFlags::AnyDynamic))
 	{
-		check(RHICmdList);
-		Buffer->UploadResourceData(*RHICmdList, CreateInfo.ResourceArray, DesiredState, CreateInfo.DebugName, TraceClassName, CreateInfo.OwnerName);
+		// Copy directly in mapped data
+		const FD3D12ResourceLocation& UploadLocation = Buffer->ResourceLocation;
+		return UE::RHICore::FDefaultBufferInitializer(RHICmdList, Buffer, UploadLocation.GetMappedBaseAddress(), UploadLocation.GetSize());
 	}
 
-	return Buffer;
+	FD3D12ResourceLocation UploadLocation(AllocateUploadMemory(Buffer, CreateDesc));
+
+	// Note: we have to get this pointer before the Initializer is constructed because it will be moved before we get the pointer
+	void* WritableData = UploadLocation.GetMappedBaseAddress();
+
+	return UE::RHICore::FCustomBufferInitializer(RHICmdList, Buffer, WritableData, CreateDesc.Size,
+		[Buffer = TRefCountPtr<FD3D12Buffer>(Buffer), DesiredState, UploadLocation = MoveTemp(UploadLocation)](FRHICommandListBase& RHICmdList) mutable
+		{
+			RHICmdList.EnqueueLambda(
+				[Buffer = Buffer.GetReference(), DesiredState, UploadLocation = MoveTemp(UploadLocation)](FRHICommandListBase& ExecutingCmdList)
+				{
+					const FRHIGPUMask EffectiveMask = ExecutingCmdList.GetGPUMask();
+					for (uint32 GPUIndex : EffectiveMask)
+					{
+						FD3D12CommandContext& CommandContext = FD3D12CommandContext::Get(ExecutingCmdList, GPUIndex);
+						FD3D12Buffer* DeviceBuffer = CommandContext.RetrieveObject<FD3D12Buffer>(Buffer, GPUIndex);
+
+						Buffer->UploadResourceData(CommandContext, DesiredState, DeviceBuffer->ResourceLocation, UploadLocation, Buffer->GetSize());
+					}
+				});
+
+			return TRefCountPtr<FRHIBuffer>(MoveTemp(Buffer));
+		});
 }
 
-FRHIBuffer* FD3D12DynamicRHI::CreateBuffer(const FRHIBufferCreateInfo& CreateInfo, const TCHAR* DebugName, ERHIAccess InitialState, ID3D12ResourceAllocator* ResourceAllocator)
+FRHIBufferInitializer FD3D12DynamicRHI::RHICreateBufferInitializer(FRHICommandListBase& RHICmdList, const FRHIBufferCreateDesc& CreateDesc)
 {
-	FRHIResourceCreateInfo ResourceCreateInfo(DebugName);
-	return CreateD3D12Buffer(nullptr, FRHIBufferDesc(CreateInfo.Size, CreateInfo.Stride, CreateInfo.Usage), InitialState, ResourceCreateInfo, ResourceAllocator);
+	UE_TRACE_METADATA_SCOPE_ASSET_FNAME(CreateDesc.DebugName, GetRHIBufferClassName(CreateDesc.ClassName), CreateDesc.OwnerName);
+
+	if (CreateDesc.IsNull())
+	{
+		FD3D12Buffer* Buffer = GetAdapter().CreateLinkedObject<FD3D12Buffer>(
+			CreateDesc.GPUMask,
+			[CreateDesc](FD3D12Device* Device, FD3D12Buffer* FirstLinkedObject)
+			{
+				return new FD3D12Buffer(Device, CreateDesc);
+			});
+
+		return UE::RHICore::FDefaultBufferInitializer(RHICmdList, Buffer);
+	}
+
+#if D3D12RHI_PLATFORM_HAS_UNIFIED_MEMORY
+	// Unified platforms don't need to use the copy queue for uploads.
+	constexpr bool bCreateAsCopyDest = false;
+	FCreateBufferInternalResult CreateResult = CreateBufferInternal(CreateDesc, bCreateAsCopyDest, nullptr);
+
+	const FD3D12ResourceLocation& UploadLocation = CreateResult.Buffer->ResourceLocation;
+
+	void* const WritableData = reinterpret_cast<void*>(UploadLocation.GetGPUVirtualAddress());
+	const uint64 WritableSize = UploadLocation.GetSize();
+
+	return UE::RHICore::CreateUnifiedMemoryBufferInitializer(RHICmdList, CreateDesc, CreateResult.Buffer, WritableData);
+#else
+	const bool bCreateAsCopyDest = (CreateDesc.InitAction != ERHIBufferInitAction::Default);
+
+	FCreateBufferInternalResult CreateResult = CreateBufferInternal(CreateDesc, bCreateAsCopyDest, nullptr);
+
+	if (CreateDesc.InitAction == ERHIBufferInitAction::Default)
+	{
+		// Just return the buffer with its default contents
+		return UE::RHICore::FDefaultBufferInitializer(RHICmdList, CreateResult.Buffer);
+	}
+
+	// TODO: the UploadLocation logic is too complex to "just" write a different code path for each initializer type, so we will just use a common Initializer.
+	if (CreateDesc.InitAction == ERHIBufferInitAction::ResourceArray)
+	{
+		FRHIBufferInitializer Initializer = CreateBufferInitializerForWriting(RHICmdList, CreateResult, CreateDesc);
+
+		// Resource array initialization goes through the same steps as external code that writes its own data.
+		FResourceArrayUploadInterface* InitialData = CreateDesc.InitialData;
+		Initializer.WriteData(InitialData->GetResourceData(), InitialData->GetResourceDataSize());
+
+		// Discard the resource array's contents.
+		InitialData->Discard();
+
+		return MoveTemp(Initializer);
+	}
+	
+	if (CreateDesc.InitAction == ERHIBufferInitAction::Zeroed)
+	{
+		FRHIBufferInitializer Initializer = CreateBufferInitializerForWriting(RHICmdList, CreateResult, CreateDesc);
+
+		// TODO: write a custom Initializer method that enqueues zeroing on the upload context.
+		Initializer.FillWithValue(0);
+
+		return MoveTemp(Initializer);
+	}
+
+	if (CreateDesc.InitAction == ERHIBufferInitAction::Initializer)
+	{
+		return CreateBufferInitializerForWriting(RHICmdList, CreateResult, CreateDesc);
+	}
+
+	return UE::RHICore::HandleUnknownBufferInitializerInitAction(RHICmdList, CreateDesc);
+#endif // D3D12RHI_PLATFORM_HAS_UNIFIED_MEMORY
 }
 
 void* FD3D12DynamicRHI::LockBuffer(FRHICommandListBase& RHICmdList, FD3D12Buffer* Buffer, uint32 BufferSize, EBufferUsageFlags BufferUsage, uint32 Offset, uint32 Size, EResourceLockMode LockMode)
@@ -536,7 +532,7 @@ void* FD3D12DynamicRHI::LockBuffer(FRHICommandListBase& RHICmdList, FD3D12Buffer
 	void* Data = nullptr;
 
 	// Determine whether the buffer is dynamic or not.
-	if (EnumHasAnyFlags(BufferUsage, BUF_AnyDynamic))
+	if (EnumHasAnyFlags(BufferUsage, EBufferUsageFlags::AnyDynamic))
 	{
 		check(LockMode == RLM_WriteOnly || LockMode == RLM_WriteOnly_NoOverwrite);
 
@@ -651,7 +647,7 @@ void FD3D12DynamicRHI::UnlockBuffer(FRHICommandListBase& RHICmdList, FD3D12Buffe
 	check(LockedData.bLocked == true);
 
 	// Determine whether the buffer is dynamic or not.
-	if (EnumHasAnyFlags(BufferUsage, BUF_AnyDynamic))
+	if (EnumHasAnyFlags(BufferUsage, EBufferUsageFlags::AnyDynamic))
 	{
 		// If the Buffer is dynamic, its upload heap memory can always stay mapped. Don't do anything.
 	}
@@ -712,7 +708,7 @@ void* FD3D12DynamicRHI::RHILockBuffer(FRHICommandListBase& RHICmdList, FRHIBuffe
 	// If you hit this assert, you should be using LockBufferMGPU and iterating over FRHIGPUMask::All() to initialize the resource separately for each GPU.
 	// "MultiGPUAllocate" only makes sense if a buffer must vary per GPU, for example if it's a buffer that includes GPU specific virtual addresses for ray
 	// tracing acceleration structures.
-	check(!EnumHasAnyFlags(BufferRHI->GetUsage(), BUF_MultiGPUAllocate));
+	check(!EnumHasAnyFlags(BufferRHI->GetUsage(), EBufferUsageFlags::MultiGPUAllocate));
 
 	bool bNeedTransition = ( (!(EnumHasAnyFlags(BufferRHI->GetUsage(), BUF_AnyDynamic))) && (LockMode == EResourceLockMode::RLM_ReadOnly) );
 
@@ -737,7 +733,7 @@ void* FD3D12DynamicRHI::RHILockBufferMGPU(FRHICommandListBase& RHICmdList, FRHIB
 	// If you hit this assert, you should be using LockBuffer to initialize the resource, rather than this function.  The MGPU version is only for resources
 	// with the MultiGPUAllocate flag, where it's necessary for the caller to initialize the buffer for each GPU.  The other LockBuffer call initializes the
 	// resource on all GPUs with one call, due to driver mirroring of the underlying resource.
-	check(EnumHasAnyFlags(BufferRHI->GetUsage(), BUF_MultiGPUAllocate));
+	check(EnumHasAnyFlags(BufferRHI->GetUsage(), EBufferUsageFlags::MultiGPUAllocate));
 
 	FD3D12Buffer* Buffer = FD3D12DynamicRHI::ResourceCast(BufferRHI, GPUIndex);
 	return LockBuffer(RHICmdList, Buffer, Buffer->GetSize(), Buffer->GetUsage(), Offset, Size, LockMode);
@@ -745,7 +741,7 @@ void* FD3D12DynamicRHI::RHILockBufferMGPU(FRHICommandListBase& RHICmdList, FRHIB
 
 void FD3D12DynamicRHI::RHIUnlockBuffer(FRHICommandListBase& RHICmdList, FRHIBuffer* BufferRHI)
 {
-	check(!EnumHasAnyFlags(BufferRHI->GetUsage(), BUF_MultiGPUAllocate));
+	check(!EnumHasAnyFlags(BufferRHI->GetUsage(), EBufferUsageFlags::MultiGPUAllocate));
 
 	FD3D12Buffer* Buffer = FD3D12DynamicRHI::ResourceCast(BufferRHI);
 	FD3D12LockedResource& LockedData = Buffer->LockedData;
@@ -767,7 +763,7 @@ void FD3D12DynamicRHI::RHIUnlockBuffer(FRHICommandListBase& RHICmdList, FRHIBuff
 
 void FD3D12DynamicRHI::RHIUnlockBufferMGPU(FRHICommandListBase& RHICmdList, FRHIBuffer* BufferRHI, uint32 GPUIndex)
 {
-	check(EnumHasAnyFlags(BufferRHI->GetUsage(), BUF_MultiGPUAllocate));
+	check(EnumHasAnyFlags(BufferRHI->GetUsage(), EBufferUsageFlags::MultiGPUAllocate));
 
 	FD3D12Buffer* Buffer = FD3D12DynamicRHI::ResourceCast(BufferRHI, GPUIndex);
 	UnlockBuffer(RHICmdList, Buffer, Buffer->GetUsage());
diff --git a/Engine/Source/Runtime/D3D12RHI/Private/D3D12Device.cpp b/Engine/Source/Runtime/D3D12RHI/Private/D3D12Device.cpp
--- a/Engine/Source/Runtime/D3D12RHI/Private/D3D12Device.cpp
+++ b/Engine/Source/Runtime/D3D12RHI/Private/D3D12Device.cpp
@@ -412,12 +412,21 @@ void FD3D12Device::SetupAfterDeviceCreation()
 	RayTracingCompactionRequestHandler = new FD3D12RayTracingCompactionRequestHandler(this);
 	for (FD3D12Queue& Queue : Queues)
 	{
-		D3D12_RESOURCE_DESC DispatchRaysDescBufferDesc = CD3DX12_RESOURCE_DESC::Buffer(sizeof(D3D12_DISPATCH_RAYS_DESC), D3D12RHI_RESOURCE_FLAG_ALLOW_INDIRECT_BUFFER | D3D12_RESOURCE_FLAG_ALLOW_UNORDERED_ACCESS);
+		const size_t DispatchRaysSize = sizeof(D3D12_DISPATCH_RAYS_DESC);
+
+		const D3D12_RESOURCE_DESC DispatchRaysBufferDesc = CD3DX12_RESOURCE_DESC::Buffer(DispatchRaysSize, D3D12RHI_RESOURCE_FLAG_ALLOW_INDIRECT_BUFFER | D3D12_RESOURCE_FLAG_ALLOW_UNORDERED_ACCESS);
+		const FRHIBufferCreateDesc CreateDesc =
+			FRHIBufferCreateDesc::Create(TEXT("DispatchRaysDescBuffer"), DispatchRaysSize, 0, EBufferUsageFlags::DrawIndirect)
+			.SetGPUMask(GetGPUMask())
+			.SetInitialState(ERHIAccess::IndirectArgs);
+
 		Queue.RayTracingDispatchRaysDescBuffer = GetParentAdapter()->CreateRHIBuffer(
-			DispatchRaysDescBufferDesc, 256,
-			FRHIBufferDesc(DispatchRaysDescBufferDesc.Width, 0, BUF_DrawIndirect),
-			ED3D12ResourceStateMode::MultiState, D3D12_RESOURCE_STATE_INDIRECT_ARGUMENT, false /*bInitialData*/,
-			GetGPUMask(), nullptr /*ResourceAllocator*/, TEXT("DispatchRaysDescBuffer"), ERHIAccess::IndirectArgs);
+			DispatchRaysBufferDesc,
+			256,
+			CreateDesc,
+			ED3D12ResourceStateMode::MultiState,
+			D3D12_RESOURCE_STATE_INDIRECT_ARGUMENT
+		);
 	}
 #endif // D3D12_RHI_RAYTRACING
 
