Add shader binding layout support for static root signature code and serialize the shader binding layout on consoles as well

diff --git a/Engine/Source/Runtime/D3D12RHI/Private/D3D12Adapter.cpp b/Engine/Source/Runtime/D3D12RHI/Private/D3D12Adapter.cpp
--- a/Engine/Source/Runtime/D3D12RHI/Private/D3D12Adapter.cpp
+++ b/Engine/Source/Runtime/D3D12RHI/Private/D3D12Adapter.cpp
@@ -1317,6 +1317,12 @@ void FD3D12Adapter::InitializeDevices()
 			EnumAddFlags(RayTracingFlags, EShaderBindingLayoutFlags::BindlessSamplers);
 		}
 
+		if (RHIGetStaticShaderBindingLayoutSupport(GMaxRHIShaderPlatform) != ERHIStaticShaderBindingLayoutSupport::Unsupported)
+		{
+			check(EnumHasAllFlags(RayTracingFlags, EShaderBindingLayoutFlags::BindlessResources | EShaderBindingLayoutFlags::BindlessSamplers))
+			EnumAddFlags(RayTracingFlags, EShaderBindingLayoutFlags::ShaderBindingLayoutUsed);
+		}
+
 		StaticRayTracingGlobalRootSignature.InitStaticRayTracingGlobalRootSignatureDesc(RayTracingFlags);
 		StaticRayTracingLocalRootSignature.InitStaticRayTracingLocalRootSignatureDesc(RayTracingFlags);
 #endif
diff --git a/Engine/Source/Runtime/D3D12RHI/Private/D3D12RootSignature.cpp b/Engine/Source/Runtime/D3D12RHI/Private/D3D12RootSignature.cpp
--- a/Engine/Source/Runtime/D3D12RHI/Private/D3D12RootSignature.cpp
+++ b/Engine/Source/Runtime/D3D12RHI/Private/D3D12RootSignature.cpp
@@ -722,7 +722,24 @@ void FD3D12RootSignature::InternalAnalyzeSignature(const RootSignatureDescType&
 		}
 
 		const uint32 BindingSpace = GetBindingSpace(InRootSignatureType, CurrentVisibleSF);
-		if (BindingSpace != ParameterBindingSpace)
+
+		// Track parameters bound to the Static Shader Binding layout slot independent of the regular parameter binding space
+		if (ParameterBindingSpace == UE_HLSL_SPACE_STATIC_SHADER_BINDINGS)
+		{
+			if (StaticShaderBindingSlot < 0)
+			{
+				// This is the first CBV for this shader binding layout, save it's root parameter index (other CBVs will be indexed using this base root parameter index).
+				StaticShaderBindingSlot = i;
+				StaticShaderBindingCount = 0;
+			}
+			StaticShaderBindingCount++;
+
+			// All shader binding layout CBVs should be contiguous registers.
+			check(i == StaticShaderBindingSlot + CurrentParameter.Descriptor.ShaderRegister);
+
+			continue;
+		}
+		else if (BindingSpace != ParameterBindingSpace)
 		{
 			// Only consider parameters in the requested binding space.
 			continue;
diff --git a/Engine/Source/Runtime/D3D12RHI/Public/D3D12RootSignatureDefinitions.h b/Engine/Source/Runtime/D3D12RHI/Public/D3D12RootSignatureDefinitions.h
--- a/Engine/Source/Runtime/D3D12RHI/Public/D3D12RootSignatureDefinitions.h
+++ b/Engine/Source/Runtime/D3D12RHI/Public/D3D12RootSignatureDefinitions.h
@@ -484,6 +484,13 @@ namespace D3D12ShaderUtils
 				Creator.AddShaderResourceViewParameter(RAY_TRACING_SYSTEM_VERTEXBUFFER_REGISTER, UE_HLSL_SPACE_RAY_TRACING_SYSTEM);
 			}
 		}
+		else if (Creator.HasFlags(EShaderBindingLayoutFlags::ShaderBindingLayoutUsed))
+		{
+			for (uint32 Index = 0; Index < FRHIShaderBindingLayout::MaxUniformBufferEntries; ++Index)
+			{
+				Creator.AddConstantBufferViewParameter(Index, UE_HLSL_SPACE_STATIC_SHADER_BINDINGS);
+			}
+		}
 
 		if (Creator.HasFlags(EShaderBindingLayoutFlags::BindlessResources))
 		{
diff --git a/Engine/Source/Runtime/Engine/Private/ShaderCompiler/ShaderCompiler.cpp b/Engine/Source/Runtime/Engine/Private/ShaderCompiler/ShaderCompiler.cpp
--- a/Engine/Source/Runtime/Engine/Private/ShaderCompiler/ShaderCompiler.cpp
+++ b/Engine/Source/Runtime/Engine/Private/ShaderCompiler/ShaderCompiler.cpp
@@ -3940,6 +3940,11 @@ void GlobalBeginCompileShader(
 		}
 	}
 
+	if (Input.Environment.ShaderBindingLayout)
+	{
+ 		Input.Environment.CompilerFlags.Add(CFLAG_ShaderBindingLayout);
+	}
+
 	SET_SHADER_DEFINE(Input.Environment, PLATFORM_SUPPORTS_BINDLESS, bEnableBindlessMacro);
 
 	if (CVarShadersRemoveDeadCode.GetValueOnAnyThread())
diff --git a/Engine/Source/Runtime/RHI/Public/RHIShaderBindingLayout.h b/Engine/Source/Runtime/RHI/Public/RHIShaderBindingLayout.h
--- a/Engine/Source/Runtime/RHI/Public/RHIShaderBindingLayout.h
+++ b/Engine/Source/Runtime/RHI/Public/RHIShaderBindingLayout.h
@@ -57,6 +57,7 @@ enum class EShaderBindingLayoutFlags : uint8
 	BindlessResources = 1 << 2,
 	BindlessSamplers = 1 << 3,
 	RootConstants = 1 << 4,
+	ShaderBindingLayoutUsed = 1 << 5,
 };
 ENUM_CLASS_FLAGS(EShaderBindingLayoutFlags)
 
diff --git a/Engine/Source/Runtime/RenderCore/Public/ShaderCompilerCore.h b/Engine/Source/Runtime/RenderCore/Public/ShaderCompilerCore.h
--- a/Engine/Source/Runtime/RenderCore/Public/ShaderCompilerCore.h
+++ b/Engine/Source/Runtime/RenderCore/Public/ShaderCompilerCore.h
@@ -168,6 +168,8 @@ enum ECompilerFlags
 	CFLAG_ShaderBundle,
 	// Shader code should not be stripped of comments/whitespace/line directives at the end of preprocessing
 	CFLAG_DisableSourceStripping,
+	// Shader uses RHI Shader Binding Layout for global shader binding
+	CFLAG_ShaderBindingLayout,
 	CFLAG_Max,
 };
 static_assert(CFLAG_Max < 64, "Out of bitfield space! Modify FShaderCompilerFlags");
