Inline raytracing support code for d3d12 rhi raytracing code

diff --git a/Engine/Shaders/Private/RayTracing/TraceRayInline.ush b/Engine/Shaders/Private/RayTracing/TraceRayInline.ush
--- a/Engine/Shaders/Private/RayTracing/TraceRayInline.ush
+++ b/Engine/Shaders/Private/RayTracing/TraceRayInline.ush
@@ -31,6 +31,8 @@
 #include "/Engine/Private/RayTracing/TraceRayInlineVulkan.ush"
 #elif COMPILER_METAL
 // #include "/Engine/Private/RayTracing/TraceRayInlineMetal.ush"
+#elif SM6_PROFILE || SM5_PROFILE
+#include "/Engine/Private/RayTracing/TraceRayInlineD3D12.ush"
 #endif // COMPILER_VULKAN
 
 #ifndef PLATFORM_SUPPORTS_INLINE_RAY_TRACING_TRIANGLE_NORMALS
diff --git a/Engine/Shaders/Private/RayTracing/TraceRayInlineD3D12.ush b/Engine/Shaders/Private/RayTracing/TraceRayInlineD3D12.ush
new file mode 100644
--- /dev/null
+++ b/Engine/Shaders/Private/RayTracing/TraceRayInlineD3D12.ush
@@ -0,0 +1,58 @@
+// Copyright Epic Games, Inc. All Rights Reserved.
+
+/*==============================================================================
+	TraceRayInlineD3D12.ush: D3D12-specific inline ray tracing functionality,
+	such as utility functions for accessing triangle indices and vertices.
+==============================================================================*/
+
+#pragma once
+
+#include "/Engine/Shared/RayTracingDefinitions.h"
+#include "/Engine/Shared/RayTracingBuiltInResources.h"
+#include "/Engine/Private/RayTracing/RayTracingCommon.ush"
+#include "/Engine/Private/RayTracing/RayTracingHitGroupCommon.ush"
+
+#if COMPILER_DXC && PLATFORM_SUPPORTS_INLINE_RAY_TRACING && ENABLE_BINDLESS_RESOURCES && PLATFORM_SUPPORTS_BINDLESS
+
+#define PLATFORM_SUPPORTS_INLINE_RAY_TRACING_TRIANGLE_ATTRIBUTES 1
+
+struct FD3DHitGroupSystemParametersWithPadding
+{
+	FHitGroupSystemRootConstants RootConstants;
+
+	uint BindlessHitGroupSystemIndexBuffer;
+	uint BindlessHitGroupSystemVertexBuffer;
+
+	// Add padding because c++ side has union for GPU address of index/vertex buffer data for non bindless code path which are uint64 each
+	uint Padding0;
+	uint Padding1;
+};
+
+// Create a generic alias for the Vulkan-specific metadata type
+#define FRayTracingSceneMetadataRecord FD3DHitGroupSystemParametersWithPadding
+
+FTriangleBaseAttributes LoadInlineRayTracingTriangleAttributes(
+	StructuredBuffer<FRayTracingSceneMetadataRecord> RayTracingSceneMetadata,
+	uint InstanceContributionToHitGroupIndex,
+	uint MultiplierForGeometryContributionToHitGroupIndex,
+	uint GeometryIndex,
+	uint PrimitiveIndex)
+{
+	uint RecordIndex = InstanceContributionToHitGroupIndex / MultiplierForGeometryContributionToHitGroupIndex + GeometryIndex;
+	FRayTracingSceneMetadataRecord Params = RayTracingSceneMetadata[RecordIndex];
+		
+	const uint IndexBufferOffsetInBytes = Params.RootConstants.IndexBufferOffsetInBytes;
+	const uint IndexBufferStride = Params.RootConstants.GetIndexStride();
+	const uint VertexStride = Params.RootConstants.GetVertexStride();
+	const uint VertexBufferOffsetInBytes = 0; // Base offset is already applied when the buffer is bound to the hit group
+	
+	const uint IndexBufferHandle = Params.BindlessHitGroupSystemIndexBuffer;
+	const uint VertexBufferHandle = Params.BindlessHitGroupSystemVertexBuffer;
+
+	ByteAddressBuffer IndexBuffer = ResourceDescriptorHeap[IndexBufferHandle];
+	ByteAddressBuffer VertexBuffer = ResourceDescriptorHeap[VertexBufferHandle];	
+
+	return LoadTriangleBaseAttributes(IndexBuffer, IndexBufferOffsetInBytes, IndexBufferStride, VertexBuffer, VertexBufferOffsetInBytes, VertexStride, PrimitiveIndex);
+}
+
+#endif // COMPILER_DXC && PLATFORM_SUPPORTS_INLINE_RAY_TRACING && ENABLE_BINDLESS_RESOURCES && PLATFORM_SUPPORTS_BINDLESS
diff --git a/Engine/Shaders/Public/Platform/D3D/PlatformBindlessResources.ush b/Engine/Shaders/Public/Platform/D3D/PlatformBindlessResources.ush
--- a/Engine/Shaders/Public/Platform/D3D/PlatformBindlessResources.ush
+++ b/Engine/Shaders/Public/Platform/D3D/PlatformBindlessResources.ush
@@ -2,7 +2,7 @@
 
 #pragma once
 
-#if ENABLE_BINDLESS_RESOURCES && PLATFORM_SUPPORTS_BINDLESS && (RAYCALLABLESHADER || RAYHITGROUPSHADER || RAYGENSHADER || RAYMISSSHADER)
+#if ENABLE_BINDLESS_RESOURCES && PLATFORM_SUPPORTS_BINDLESS && (RAYCALLABLESHADER || RAYHITGROUPSHADER || RAYGENSHADER || RAYMISSSHADER || COMPUTESHADER)
 	#define OVERRIDE_RAY_TRACING_HIT_GROUP_SYSTEM_RESOURCES 1
 	#include "/Engine/Shared/RayTracingBuiltInResources.h"
 
diff --git a/Engine/Source/Runtime/D3D12RHI/Private/D3D12RayTracing.cpp b/Engine/Source/Runtime/D3D12RHI/Private/D3D12RayTracing.cpp
--- a/Engine/Source/Runtime/D3D12RHI/Private/D3D12RayTracing.cpp
+++ b/Engine/Source/Runtime/D3D12RHI/Private/D3D12RayTracing.cpp
@@ -1344,73 +1344,108 @@ public:
 	// Scaling beyond 5 total threads does not yield any speedup in practice.
 	static constexpr uint32 MaxBindingWorkers = 5; // RHI thread + 4 parallel workers.
 
-	FD3D12RayTracingShaderBindingTableInternal(const FRayTracingShaderBindingTableInitializer& Initializer, FD3D12Device* InDevice)
+	FD3D12RayTracingShaderBindingTableInternal(FRHICommandListBase& RHICmdList, const FRayTracingShaderBindingTableInitializer& Initializer, FD3D12Device* InDevice)
 		: UniqueId(NextUniqueId++)
 	{
 		checkf(Initializer.LocalBindingDataSize <= 4096, TEXT("The maximum size of a local root signature is 4KB.")); // as per section 4.22.1 of DXR spec v1.0
-		check(Initializer.ShaderBindingMode == ERayTracingShaderBindingMode::RTPSO); //< Only support RTPSO for now
-
-		const uint32 NumHitGroupSlots = Initializer.HitGroupIndexingMode == ERayTracingHitGroupIndexingMode::Allow ? Initializer.NumGeometrySegments * Initializer.NumShaderSlotsPerGeometrySegment : 1;
-		checkf(Initializer.LocalBindingDataSize >= sizeof(FD3D12HitGroupSystemParameters), TEXT("All local root signatures are expected to contain ray tracing system root parameters"));
+		check(Initializer.ShaderBindingMode != ERayTracingShaderBindingMode::Disabled);
 		
 		Device = InDevice;
 
 		Lifetime = Initializer.Lifetime;
 		HitGroupIndexingMode = Initializer.HitGroupIndexingMode;
-		LocalRecordSizeUnaligned = ShaderIdentifierSize + Initializer.LocalBindingDataSize;
-		LocalRecordStride = RoundUpToNextMultiple(LocalRecordSizeUnaligned, D3D12_RAYTRACING_SHADER_RECORD_BYTE_ALIGNMENT);
+		ShaderBindingMode = Initializer.ShaderBindingMode;
+		NumShaderSlotsPerGeometrySegment = Initializer.NumShaderSlotsPerGeometrySegment;
+
+		const uint32 NumHitGroupSlots = Initializer.HitGroupIndexingMode == ERayTracingHitGroupIndexingMode::Allow ? Initializer.NumGeometrySegments * Initializer.NumShaderSlotsPerGeometrySegment : 1;
 
-		// Custom descriptor cache is only required when local resources may be bound.
-		// If only global resources are used, then transient descriptor cache can be used.
-		const bool bNeedsDescriptorCache = (NumHitGroupSlots + Initializer.NumCallableShaderSlots + Initializer.NumMissShaderSlots) * Initializer.LocalBindingDataSize != 0;
+		NumMissRecords = Initializer.NumMissShaderSlots;
+		NumHitRecords = NumHitGroupSlots;
+		NumCallableRecords = Initializer.NumCallableShaderSlots;
 
-		if (bNeedsDescriptorCache)
+		if (EnumHasAnyFlags(ShaderBindingMode, ERayTracingShaderBindingMode::RTPSO))
 		{
-			// #dxr_todo UE-72158: Remove this when RT descriptors are sub-allocated from the global view descriptor heap.
+			checkf(Initializer.LocalBindingDataSize >= sizeof(FD3D12HitGroupSystemParameters), TEXT("All local root signatures are expected to contain ray tracing system root parameters"));
+
+			LocalRecordSizeUnaligned = ShaderIdentifierSize + Initializer.LocalBindingDataSize;
+			LocalRecordStride = RoundUpToNextMultiple(LocalRecordSizeUnaligned, D3D12_RAYTRACING_SHADER_RECORD_BYTE_ALIGNMENT);
 
-			if (GD3D12ExplicitViewDescriptorHeapOverflowReported)
+			// Custom descriptor cache is only required when local resources may be bound.
+			// If only global resources are used, then transient descriptor cache can be used.
+			const bool bNeedsDescriptorCache = (NumHitGroupSlots + Initializer.NumCallableShaderSlots + Initializer.NumMissShaderSlots) * Initializer.LocalBindingDataSize != 0;
+
+			if (bNeedsDescriptorCache)
 			{
-				GD3D12ExplicitViewDescriptorHeapSize = GD3D12ExplicitViewDescriptorHeapSize * 2;
-				GD3D12ExplicitViewDescriptorHeapOverflowReported = 0;
-			}
+				// #dxr_todo UE-72158: Remove this when RT descriptors are sub-allocated from the global view descriptor heap.
 
-			// D3D12 is guaranteed to support 1M (D3D12_MAX_SHADER_VISIBLE_DESCRIPTOR_HEAP_SIZE_TIER_1) descriptors in a CBV/SRV/UAV heap, so clamp the size to this.
-			// https://docs.microsoft.com/en-us/windows/desktop/direct3d12/hardware-support
-			const uint32 NumViewDescriptors = FMath::Min(D3D12_MAX_SHADER_VISIBLE_DESCRIPTOR_HEAP_SIZE_TIER_1, GD3D12ExplicitViewDescriptorHeapSize);
-			const uint32 NumSamplerDescriptors = D3D12_MAX_SHADER_VISIBLE_SAMPLER_HEAP_SIZE;
+				if (GD3D12ExplicitViewDescriptorHeapOverflowReported)
+				{
+					GD3D12ExplicitViewDescriptorHeapSize = GD3D12ExplicitViewDescriptorHeapSize * 2;
+					GD3D12ExplicitViewDescriptorHeapOverflowReported = 0;
+				}
 
-			DescriptorCache = new FD3D12ExplicitDescriptorCache(Device, MaxBindingWorkers);
-			DescriptorCache->Init(0, NumViewDescriptors, NumSamplerDescriptors, ERHIBindlessConfiguration::RayTracingShaders);
-		}
+				// D3D12 is guaranteed to support 1M (D3D12_MAX_SHADER_VISIBLE_DESCRIPTOR_HEAP_SIZE_TIER_1) descriptors in a CBV/SRV/UAV heap, so clamp the size to this.
+				// https://docs.microsoft.com/en-us/windows/desktop/direct3d12/hardware-support
+				const uint32 NumViewDescriptors = FMath::Min(D3D12_MAX_SHADER_VISIBLE_DESCRIPTOR_HEAP_SIZE_TIER_1, GD3D12ExplicitViewDescriptorHeapSize);
+				const uint32 NumSamplerDescriptors = D3D12_MAX_SHADER_VISIBLE_SAMPLER_HEAP_SIZE;
 
-		NumMissRecords = Initializer.NumMissShaderSlots;
-		NumHitRecords = NumHitGroupSlots;
-		NumCallableRecords = Initializer.NumCallableShaderSlots;
+				DescriptorCache = new FD3D12ExplicitDescriptorCache(Device, MaxBindingWorkers);
+				DescriptorCache->Init(0, NumViewDescriptors, NumSamplerDescriptors, ERHIBindlessConfiguration::RayTracingShaders);
+			}
 
-		uint32 TotalDataSize = 0;
-		
-		HitGroupShaderTableOffset = TotalDataSize;
-		TotalDataSize += NumHitGroupSlots * LocalRecordStride;
-		TotalDataSize = RoundUpToNextMultiple(TotalDataSize, D3D12_RAYTRACING_SHADER_TABLE_BYTE_ALIGNMENT);
+			uint32 TotalDataSize = 0;
+
+			HitGroupShaderTableOffset = TotalDataSize;
+			TotalDataSize += NumHitGroupSlots * LocalRecordStride;
+			TotalDataSize = RoundUpToNextMultiple(TotalDataSize, D3D12_RAYTRACING_SHADER_TABLE_BYTE_ALIGNMENT);
 
-		CallableShaderTableOffset = TotalDataSize;
-		TotalDataSize += Initializer.NumCallableShaderSlots * LocalRecordStride;
-		TotalDataSize = RoundUpToNextMultiple(TotalDataSize, D3D12_RAYTRACING_SHADER_TABLE_BYTE_ALIGNMENT);
+			CallableShaderTableOffset = TotalDataSize;
+			TotalDataSize += Initializer.NumCallableShaderSlots * LocalRecordStride;
+			TotalDataSize = RoundUpToNextMultiple(TotalDataSize, D3D12_RAYTRACING_SHADER_TABLE_BYTE_ALIGNMENT);
 
-		MissShaderTableOffset = TotalDataSize;
-		TotalDataSize += Initializer.NumMissShaderSlots * LocalRecordStride;
-		TotalDataSize = RoundUpToNextMultiple(TotalDataSize, D3D12_RAYTRACING_SHADER_TABLE_BYTE_ALIGNMENT);
+			MissShaderTableOffset = TotalDataSize;
+			TotalDataSize += Initializer.NumMissShaderSlots * LocalRecordStride;
+			TotalDataSize = RoundUpToNextMultiple(TotalDataSize, D3D12_RAYTRACING_SHADER_TABLE_BYTE_ALIGNMENT);
 
-		Data.SetNumZeroed(TotalDataSize);
+			Data.SetNumZeroed(TotalDataSize);
 #if DO_CHECK
-		bWasDefaultMissShaderSet = false;
+			bWasDefaultMissShaderSet = false;
 #endif
-		SetDefaultHitGroupIdentifier(FD3D12ShaderIdentifier::Null);
-		SetDefaultMissShaderIdentifier(FD3D12ShaderIdentifier::Null);
-		SetDefaultCallableShaderIdentifier(FD3D12ShaderIdentifier::Null);
+			SetDefaultHitGroupIdentifier(FD3D12ShaderIdentifier::Null);
+			SetDefaultMissShaderIdentifier(FD3D12ShaderIdentifier::Null);
+			SetDefaultCallableShaderIdentifier(FD3D12ShaderIdentifier::Null);
+
+			// Keep CPU-side data after upload
+			Data.SetAllowCPUAccess(true);
+		}
 
-		// Keep CPU-side data after upload
-		Data.SetAllowCPUAccess(true);
+		if (EnumHasAnyFlags(ShaderBindingMode, ERayTracingShaderBindingMode::Inline) && Initializer.NumGeometrySegments > 0)
+		{
+			// Doesn't make sense to have inline SBT without hitgroup indexing
+			check(HitGroupIndexingMode == ERayTracingHitGroupIndexingMode::Allow);
+
+			const uint32 ParameterBufferSize = Initializer.NumGeometrySegments * sizeof(FD3D12HitGroupSystemParameters);
+			InlineGeometryParameterData.SetNumZeroed(ParameterBufferSize);
+		
+			const D3D12_RESOURCE_DESC ResourceDesc = CD3DX12_RESOURCE_DESC::Buffer(ParameterBufferSize, D3D12_RESOURCE_FLAG_NONE, D3D12_RAYTRACING_SHADER_TABLE_BYTE_ALIGNMENT);
+
+			const FRHIBufferCreateDesc CreateDesc =
+				FRHIBufferCreateDesc::Create(TEXT("RayTracingSceneMetadata"), ResourceDesc.Width, sizeof(FD3D12HitGroupSystemParameters), BUF_Static | BUF_StructuredBuffer | BUF_ShaderResource)
+				.SetInitialState(ERHIAccess::SRVCompute)
+				.SetGPUMask(FRHIGPUMask::FromIndex(Device->GetGPUIndex()));
+
+			ID3D12ResourceAllocator* ResourceAllocator = nullptr;
+			InlineGeometryParameterBuffer = Device->GetParentAdapter()->CreateRHIBuffer(
+				ResourceDesc,
+				ResourceDesc.Alignment,
+				CreateDesc,
+				ED3D12ResourceStateMode::MultiState,
+				D3D12_RESOURCE_STATE_NON_PIXEL_SHADER_RESOURCE,
+				/*bHasInitialData*/ false
+			);
+
+			InlineGeometryParameterSRV = RHICmdList.CreateShaderResourceView(InlineGeometryParameterBuffer.GetReference(), FRHIViewDesc::CreateBufferSRV().SetType(FRHIViewDesc::EBufferType::Structured));
+		}
 
 		if (Lifetime == ERayTracingShaderBindingTableLifetime::Persistent)
 		{
@@ -1573,20 +1608,41 @@ public:
 		}
 	}
 
+	void SetInlineGeometryParameters(uint32 SegmentIndex, const void* InData, uint32 InDataSize)
+	{
+		const uint32 WriteOffset = InDataSize * SegmentIndex;
+		FMemory::Memcpy(&InlineGeometryParameterData[WriteOffset], InData, InDataSize);
+				
+		MaxUsedInlineGeometryParameterSegmentIndex = FMath::Max(MaxUsedInlineGeometryParameterSegmentIndex, (int32)SegmentIndex);
+	}
+
+	void CompareInlineGeometryParameters(uint32 SegmentIndex, const void* InData, uint32 InDataSize)
+	{
+		const uint32 WriteOffset = InDataSize * SegmentIndex;
+		const void* CurrentData = InlineGeometryParameterData.GetData() + WriteOffset;
+		check(FMemory::Memcmp(CurrentData, InData, InDataSize) == 0);
+	}
+
+	FRHIShaderResourceView* GetOrCreateInlineBufferSRV(FRHICommandListBase& RHICmdList)
+	{	
+		if (EnumHasAnyFlags(ShaderBindingMode, ERayTracingShaderBindingMode::Inline))
+		{
+			check(InlineGeometryParameterSRV.IsValid());
+			return InlineGeometryParameterSRV.GetReference();
+		}
+		else
+		{
+			return nullptr;
+		}
+	}
+
 	void Commit(FD3D12CommandContext& Context)
 	{
 		TRACE_CPUPROFILER_EVENT_SCOPE(ShaderTableCommit);
 
 		check(IsInRHIThread() || !IsRunningRHIInSeparateThread());
-
 		checkf(bIsDirty, TEXT("bIsDirty should be checked before calling Commit()"));
 
-		checkf(Data.Num(), TEXT("Shader table is expected to be initialized before copying to GPU."));
-
-#if DO_CHECK
-		checkf(bWasDefaultMissShaderSet, TEXT("At least the first miss shader must have been set before copying to GPU."));
-#endif
-
 #if ENABLE_RESIDENCY_MANAGEMENT
 		// Merge all data from worker threads into the main set
 		{
@@ -1653,35 +1709,81 @@ public:
 			WorkerData[WorkerIndex].ShaderRecordCache.Empty();
 		}
 
-		const uint32 BufferSize = Data.GetResourceDataSize();
-
 		FD3D12Adapter* Adapter = Device->GetParentAdapter();
-		D3D12_RESOURCE_DESC BufferDesc = CD3DX12_RESOURCE_DESC::Buffer(BufferSize, D3D12_RESOURCE_FLAG_NONE, D3D12_RAYTRACING_SHADER_TABLE_BYTE_ALIGNMENT);
 
-		const FRHIBufferCreateDesc CreateDesc =
-			FRHIBufferCreateDesc::Create(TEXT("ShaderBindingTable"), BufferSize, 0, BUF_Static)
+		if (EnumHasAnyFlags(ShaderBindingMode, ERayTracingShaderBindingMode::RTPSO))
+		{
+			checkf(Data.Num(), TEXT("Shader table is expected to be initialized before copying to GPU."));
+			#if DO_CHECK
+			checkf(bWasDefaultMissShaderSet, TEXT("At least the first miss shader must have been set before copying to GPU."));
+			#endif
+						
+			// TODO: only allocate for MaxUsedHitRecordIndex
+
+			const uint32 BufferSize = Data.GetResourceDataSize();
+
+			D3D12_RESOURCE_DESC BufferDesc = CD3DX12_RESOURCE_DESC::Buffer(BufferSize, D3D12_RESOURCE_FLAG_NONE, D3D12_RAYTRACING_SHADER_TABLE_BYTE_ALIGNMENT);
+
+			const FRHIBufferCreateDesc CreateDesc =
+				FRHIBufferCreateDesc::Create(TEXT("ShaderBindingTable"), BufferSize, 0, BUF_Static)
 			.SetInitialState(ERHIAccess::CopyDest)
 			.SetGPUMask(FRHIGPUMask::FromIndex(Device->GetGPUIndex()));
 
-		ID3D12ResourceAllocator* ResourceAllocator = nullptr;
-		Buffer = Adapter->CreateRHIBuffer(
-			BufferDesc,
-			BufferDesc.Alignment,
-			CreateDesc,
-			ED3D12ResourceStateMode::MultiState,
-			D3D12_RESOURCE_STATE_COPY_DEST,
-			/*bHasInitialData*/ true
-		);
+			ID3D12ResourceAllocator* ResourceAllocator = nullptr;
+			Buffer = Adapter->CreateRHIBuffer(
+				BufferDesc,
+				BufferDesc.Alignment,
+				CreateDesc,
+				ED3D12ResourceStateMode::MultiState,
+				D3D12_RESOURCE_STATE_COPY_DEST,
+				/*bHasInitialData*/ true
+			);
 
-		// Use copy queue for uploading the data
-		Context.BatchedSyncPoints.ToWait.Emplace(Buffer->UploadResourceDataViaCopyQueue(Context, &Data));
+			// Use copy queue for uploading the data
+			Context.BatchedSyncPoints.ToWait.Emplace(Buffer->UploadResourceDataViaCopyQueue(Context, &Data));
 
-		// Enqueue transition to SRV
-		Context.TransitionResource(
-			Buffer->GetResource(),
-			D3D12_RESOURCE_STATE_COPY_DEST,
-			D3D12_RESOURCE_STATE_NON_PIXEL_SHADER_RESOURCE,
-			0);
+			// Enqueue transition to SRV
+			Context.TransitionResource(
+				Buffer->GetResource(),
+				D3D12_RESOURCE_STATE_COPY_DEST,
+				D3D12_RESOURCE_STATE_NON_PIXEL_SHADER_RESOURCE,
+				0);
+		}
+
+		// Also copy geometry parameter data to the GPU buffer
+		if (EnumHasAnyFlags(ShaderBindingMode, ERayTracingShaderBindingMode::Inline) && MaxUsedInlineGeometryParameterSegmentIndex >= 0)
+		{
+			check(InlineGeometryParameterBuffer);
+			
+			// Could be moved to copy queue if needed but complicated with SRV which is needed at RenderThread timeline
+			// and persistent SBT. Buffer needs to be created and valid because the SRV needs to be available, but needs new GPU memory
+			// each commit because copy queue overlaps with render frame which is reading from the buffer. We currently have
+			// no API or use case to allocate new static GPU memory for buffers on lock/unlock, but only new upload memory only.
+			// Current solution is just lock/unlock static VRAM buffer with new upload memory and GPU timeline copy
+			
+			// Enqueue transition to CopyDest
+			Context.TransitionResource(
+				InlineGeometryParameterBuffer->GetResource(),
+				D3D12_RESOURCE_STATE_NON_PIXEL_SHADER_RESOURCE,
+				D3D12_RESOURCE_STATE_COPY_DEST,
+				0);
+
+			{
+				TRHICommandList_RecursiveHazardous<FD3D12CommandContext> RHICmdList(&Context);
+				RHICmdList.SetAllowExtraTransition(false);
+				uint32 MaxUsedInlineGeometryParamaterDataSize = (MaxUsedInlineGeometryParameterSegmentIndex + 1) * sizeof(FD3D12HitGroupSystemParameters);
+				void* MappedUploadMemory = GDynamicRHI->RHILockBuffer(RHICmdList, InlineGeometryParameterBuffer, 0, MaxUsedInlineGeometryParamaterDataSize, EResourceLockMode::RLM_WriteOnly);
+				FMemory::Memcpy(MappedUploadMemory, InlineGeometryParameterData.GetData(), MaxUsedInlineGeometryParamaterDataSize);
+				GDynamicRHI->RHIUnlockBuffer(RHICmdList, InlineGeometryParameterBuffer);
+			}
+
+			// Enqueue transition to SRV
+			Context.TransitionResource(
+				InlineGeometryParameterBuffer->GetResource(),
+				D3D12_RESOURCE_STATE_COPY_DEST,
+				D3D12_RESOURCE_STATE_NON_PIXEL_SHADER_RESOURCE,
+				0);
+		}
 
 		bIsDirty = false;
 	}
@@ -1744,6 +1846,7 @@ public:
 	uint32 NumMissRecords = 0;
 
 	int32 MaxUsedHitRecordIndex = -1;
+	int32 MaxUsedInlineGeometryParameterSegmentIndex = -1;
 
 	uint32 MissShaderTableOffset = 0;
 	uint32 HitGroupShaderTableOffset = 0;
@@ -1753,6 +1856,8 @@ public:
 
 	ERayTracingShaderBindingTableLifetime Lifetime = ERayTracingShaderBindingTableLifetime::Transient;
 	ERayTracingHitGroupIndexingMode HitGroupIndexingMode = ERayTracingHitGroupIndexingMode::Allow;
+	ERayTracingShaderBindingMode ShaderBindingMode = ERayTracingShaderBindingMode::RTPSO;
+	uint32 NumShaderSlotsPerGeometrySegment = 0;
 
 	// Note: TABLE_BYTE_ALIGNMENT is used instead of RECORD_BYTE_ALIGNMENT to allow arbitrary switching 
 	// between multiple RayGen and Miss shaders within the same underlying table.
@@ -1762,6 +1867,11 @@ public:
 	uint32 LocalRecordStride = 0; // size of shader identifier + local root parameters, aligned to SHADER_RECORD_BYTE_ALIGNMENT (same for hit groups and callable shaders)
 	TResourceArray<uint8, D3D12_RAYTRACING_SHADER_TABLE_BYTE_ALIGNMENT> Data;
 
+	// Buffer that contains per-hitrecord index and vertex buffer binding data
+	TResourceArray<uint8, D3D12_RAYTRACING_SHADER_TABLE_BYTE_ALIGNMENT> InlineGeometryParameterData;
+	TRefCountPtr<FD3D12Buffer> InlineGeometryParameterBuffer;
+	FShaderResourceViewRHIRef InlineGeometryParameterSRV;
+
 	bool bIsDirty = true;
 	TRefCountPtr<FD3D12Buffer> Buffer;
 #if DO_CHECK
@@ -3068,7 +3178,7 @@ public:
 
 	UE_NONCOPYABLE(FD3D12RayTracingShaderBindingTable)
 
-	FD3D12RayTracingShaderBindingTable(FD3D12Adapter* Adapter, const FRayTracingShaderBindingTableInitializer& InInitializer)
+	FD3D12RayTracingShaderBindingTable(FRHICommandListBase& RHICmdList, FD3D12Adapter* Adapter, const FRayTracingShaderBindingTableInitializer& InInitializer)
 		: FRHIShaderBindingTable(InInitializer), FD3D12AdapterChild(Adapter)
 	{
 		INC_DWORD_STAT(STAT_D3D12RayTracingAllocatedSBT);
@@ -3077,7 +3187,7 @@ public:
 
 		for (FD3D12Device* Device : Adapter->GetDevices())
 		{
-			InitForDevice(Device);
+			InitForDevice(RHICmdList, Device);
 		}
 	};
 
@@ -3106,9 +3216,16 @@ public:
 		ShaderTablesPerGPU[GPUIndex] = nullptr;
 	}
 
+	virtual FRHIShaderResourceView* GetOrCreateInlineBufferSRV(FRHICommandListBase& RHICmdList) override final
+	{
+		// Should be single GPU mask, otherwise iterate for all shader tables and create linked SRV view here?
+		int32 GPUIndex = RHICmdList.GetGPUMask().ToIndex();
+		return ShaderTablesPerGPU[GPUIndex]->GetOrCreateInlineBufferSRV(RHICmdList);
+	}
+
 private:
 
-	void InitForDevice(FD3D12Device* Device)
+	void InitForDevice(FRHICommandListBase& RHICmdList, FD3D12Device* Device)
 	{
 	 	TRACE_CPUPROFILER_EVENT_SCOPE(ShaderTableInit);
 	 	SCOPE_CYCLE_COUNTER(STAT_D3D12CreateShaderTable);
@@ -3116,7 +3233,7 @@ private:
 		const uint32 GPUIndex = Device->GetGPUIndex();
 
 		check(ShaderTablesPerGPU[GPUIndex] == nullptr);
-		ShaderTablesPerGPU[GPUIndex] = new FD3D12RayTracingShaderBindingTableInternal(Initializer, Device);
+		ShaderTablesPerGPU[GPUIndex] = new FD3D12RayTracingShaderBindingTableInternal(RHICmdList, Initializer, Device);
 	}
 	
 	FD3D12RayTracingShaderBindingTableInternal* ShaderTablesPerGPU[MAX_NUM_GPUS] = {};
@@ -3463,7 +3580,7 @@ FShaderBindingTableRHIRef FD3D12DynamicRHI::RHICreateShaderBindingTable(FRHIComm
 
 	FD3D12Adapter& Adapter = GetAdapter();
 
-	return new FD3D12RayTracingShaderBindingTable(&Adapter, Initializer);
+	return new FD3D12RayTracingShaderBindingTable(RHICmdList, &Adapter, Initializer);
 }
 
 FBufferRHIRef FD3D12RayTracingGeometry::NullTransformBuffer;
@@ -4624,7 +4741,8 @@ FRHIShaderBindingTable* FD3D12RayTracingScene::FindOrCreateShaderBindingTable(co
 	SBTInitializer.ShaderBindingMode = ERayTracingShaderBindingMode::RTPSO;
 	PRAGMA_ENABLE_DEPRECATION_WARNINGS
 
-	FD3D12RayTracingShaderBindingTable* CreatedShaderTable = new FD3D12RayTracingShaderBindingTable(GetParentAdapter(), MoveTemp(SBTInitializer));
+	// Create new table (use FRHICommandListExecutor::GetImmediateCommandList() directly for now this is deprecated code)
+	FD3D12RayTracingShaderBindingTable* CreatedShaderTable = new FD3D12RayTracingShaderBindingTable(FRHICommandListExecutor::GetImmediateCommandList(), GetParentAdapter(), MoveTemp(SBTInitializer));
 
 	ShaderTables.Add(Pipeline, CreatedShaderTable);
 
@@ -5819,64 +5937,90 @@ static void SetRayTracingHitGroup(
 	}
 #endif // DO_CHECK
 
-	ShaderTable->SetHitGroupSystemParameters(WorkerIndex, BindingType, RecordIndex, Geometry, UserData, GeometrySegmentIndex);
+	FD3D12HitGroupSystemParameters SystemParameters = Geometry->HitGroupSystemParameters[GPUIndex][GeometrySegmentIndex];
+	SystemParameters.RootConstants.UserData = UserData;
 
-	const FD3D12RayTracingShader* Shader = Pipeline->HitGroupShaders.Shaders[HitGroupIndex];
+	if (EnumHasAnyFlags(ShaderTable->ShaderBindingMode, ERayTracingShaderBindingMode::RTPSO))
+	{
+		ShaderTable->SetHitGroupSystemParameters(WorkerIndex, BindingType, RecordIndex, Geometry, UserData, GeometrySegmentIndex);
 
-	FD3D12RayTracingShaderBindingTableInternal::FShaderRecordCacheKey CacheKey;
+		const FD3D12RayTracingShader* Shader = Pipeline->HitGroupShaders.Shaders[HitGroupIndex];
 
-	// TODO: disable RecordCache when using persistent SBT
-	const bool bCanUseRecordCache = GRayTracingCacheShaderRecords
-		&& LooseParameterDataSize == 0 // loose parameters end up in unique constant buffers, so SBT records can't be shared
-		&& NumUniformBuffers > 0 // there is no benefit from cache if no resources are being bound
-		&& NumUniformBuffers <= CacheKey.MaxUniformBuffers
-		&& (BindingType == ERayTracingLocalShaderBindingType::Transient || ShaderTable->Lifetime == ERayTracingShaderBindingTableLifetime::Transient); // Only transient SBTs or Transient records can use the record cache because otherwise the per record listeners need to tracked and copied as well
+		FD3D12RayTracingShaderBindingTableInternal::FShaderRecordCacheKey CacheKey;
 
-	if (bCanUseRecordCache)
-	{
-		CacheKey = FD3D12RayTracingShaderBindingTableInternal::FShaderRecordCacheKey(NumUniformBuffers, UniformBuffers, HitGroupIndex);
+		// TODO: disable RecordCache when using persistent SBT
+		const bool bCanUseRecordCache = GRayTracingCacheShaderRecords
+			&& LooseParameterDataSize == 0 // loose parameters end up in unique constant buffers, so SBT records can't be shared
+			&& NumUniformBuffers > 0 // there is no benefit from cache if no resources are being bound
+			&& NumUniformBuffers <= CacheKey.MaxUniformBuffers
+			&& (BindingType == ERayTracingLocalShaderBindingType::Transient || ShaderTable->Lifetime == ERayTracingShaderBindingTableLifetime::Transient); // Only transient SBTs or Transient records can use the record cache because otherwise the per record listeners need to tracked and copied as well
 
-		uint32* ExistingRecordIndex = ShaderTable->WorkerData[WorkerIndex].ShaderRecordCache.Find(CacheKey);
-		if (ExistingRecordIndex)
+		ERayTracingLocalShaderBindingType RTPSOBindingType = BindingType;
+		if (bCanUseRecordCache)
 		{
-			// Simply copy local shader parameters from existing SBT record and set the shader identifier, skipping resource binding work.
-			const uint32 OffsetFromRootSignatureStart = sizeof(FD3D12HitGroupSystemParameters);
-			ShaderTable->SetHitGroupIdentifier(RecordIndex, Pipeline->HitGroupShaders.Identifiers[HitGroupIndex]);
-			ShaderTable->CopyHitGroupParameters(RecordIndex, *ExistingRecordIndex, OffsetFromRootSignatureStart);
+			CacheKey = FD3D12RayTracingShaderBindingTableInternal::FShaderRecordCacheKey(NumUniformBuffers, UniformBuffers, HitGroupIndex);
 
-#if DO_CHECK
-			// Validate all other data against copied data
-			BindingType = ERayTracingLocalShaderBindingType::Validation;
-#else
-			return;
-#endif // DO_CHECK
+			uint32* ExistingRecordIndex = ShaderTable->WorkerData[WorkerIndex].ShaderRecordCache.Find(CacheKey);
+			if (ExistingRecordIndex)
+			{
+				// Simply copy local shader parameters from existing SBT record and set the shader identifier, skipping resource binding work.
+				const uint32 OffsetFromRootSignatureStart = sizeof(FD3D12HitGroupSystemParameters);
+				ShaderTable->SetHitGroupIdentifier(RecordIndex, Pipeline->HitGroupShaders.Identifiers[HitGroupIndex]);
+				ShaderTable->CopyHitGroupParameters(RecordIndex, *ExistingRecordIndex, OffsetFromRootSignatureStart);
+
+				#if DO_CHECK
+				// Validate all other data against copied data
+				RTPSOBindingType = ERayTracingLocalShaderBindingType::Validation;
+				#else
+				return;
+				#endif // DO_CHECK
+			}
 		}
-	}
 
-	FD3D12RayTracingLocalResourceBinder ResourceBinder(*Device, *ShaderTable, *(Shader->LocalRootSignature), BindingType, RecordIndex, WorkerIndex, ERayTracingBindingType::HitGroup);
-	const bool bResourcesBound = SetRayTracingShaderResources(Shader, Shader->LocalRootSignature,		
-		0, nullptr,	// BindlessParameters
-		0, nullptr, // Textures
-		0, nullptr, // SRVs
-		NumUniformBuffers, UniformBuffers,
-		0, nullptr, // Samplers
-		0, nullptr, // UAVs
-		LooseParameterDataSize, LooseParameterData,
-		ResourceBinder);
+		FD3D12RayTracingLocalResourceBinder ResourceBinder(*Device, *ShaderTable, *(Shader->LocalRootSignature), RTPSOBindingType, RecordIndex, WorkerIndex, ERayTracingBindingType::HitGroup);
+		const bool bResourcesBound = SetRayTracingShaderResources(Shader, Shader->LocalRootSignature,
+		                                                          0, nullptr,	// BindlessParameters
+		                                                          0, nullptr, // Textures
+		                                                          0, nullptr, // SRVs
+		                                                          NumUniformBuffers, UniformBuffers,
+		                                                          0, nullptr, // Samplers
+		                                                          0, nullptr, // UAVs
+		                                                          LooseParameterDataSize, LooseParameterData,
+		                                                          ResourceBinder);
 
-	if (bCanUseRecordCache && bResourcesBound)
-	{
-		ShaderTable->WorkerData[WorkerIndex].ShaderRecordCache.FindOrAdd(CacheKey, RecordIndex);
-	}
+		if (bCanUseRecordCache && bResourcesBound)
+		{
+			ShaderTable->WorkerData[WorkerIndex].ShaderRecordCache.FindOrAdd(CacheKey, RecordIndex);
+		}
 
-	FD3D12ShaderIdentifier ShaderIdentifier = bResourcesBound ? Pipeline->HitGroupShaders.Identifiers[HitGroupIndex] : FD3D12ShaderIdentifier::Null;
-	if (BindingType != ERayTracingLocalShaderBindingType::Validation)
-	{
-		ShaderTable->SetHitGroupIdentifier(RecordIndex, ShaderIdentifier);
+		FD3D12ShaderIdentifier ShaderIdentifier = bResourcesBound ? Pipeline->HitGroupShaders.Identifiers[HitGroupIndex] : FD3D12ShaderIdentifier::Null;
+		if (RTPSOBindingType != ERayTracingLocalShaderBindingType::Validation)
+		{
+			ShaderTable->SetHitGroupIdentifier(RecordIndex, ShaderIdentifier);
+		}
+		else
+		{
+			ShaderTable->CompareHitGroupIdentifier(RecordIndex, ShaderIdentifier);
+		}
 	}
-	else
+
+	if (EnumHasAnyFlags(ShaderTable->ShaderBindingMode, ERayTracingShaderBindingMode::Inline))
 	{
-		ShaderTable->CompareHitGroupIdentifier(RecordIndex, ShaderIdentifier);
+		// Only care about shader slot 0 for inline geometry parameters
+		uint32 NumShaderSlotsPerGeometrySegment = ShaderTable->NumShaderSlotsPerGeometrySegment;
+		if (RecordIndex % NumShaderSlotsPerGeometrySegment == 0)
+		{			
+			// Recompute the geometry segment index from the record index
+			uint32 SegmentIndex = RecordIndex / NumShaderSlotsPerGeometrySegment;
+			if (BindingType != ERayTracingLocalShaderBindingType::Validation)
+			{
+				ShaderTable->SetInlineGeometryParameters(SegmentIndex, &SystemParameters, sizeof(FD3D12HitGroupSystemParameters));
+			}
+			else
+			{
+				ShaderTable->CompareInlineGeometryParameters(SegmentIndex, &SystemParameters, sizeof(FD3D12HitGroupSystemParameters));
+			}
+		}
 	}
 }
 
