Add a way to reserve a batch of view descriptors to FD3D12ExplicitDescriptorCache. This allows to significantly reduce the number of atomic operations when building shader binding tables in parallel.

diff --git a/Engine/Source/Runtime/D3D12RHI/Private/D3D12ExplicitDescriptorCache.cpp b/Engine/Source/Runtime/D3D12RHI/Private/D3D12ExplicitDescriptorCache.cpp
--- a/Engine/Source/Runtime/D3D12RHI/Private/D3D12ExplicitDescriptorCache.cpp
+++ b/Engine/Source/Runtime/D3D12RHI/Private/D3D12ExplicitDescriptorCache.cpp
@@ -33,6 +33,12 @@ static FAutoConsoleVariableRef CVarD3D12ExplicitViewDescriptorHeapSize(
 	ECVF_ReadOnly
 );
 
+FD3D12ExplicitDescriptorHeap::FD3D12ExplicitDescriptorHeap(FD3D12Device* Device)
+	: FD3D12DeviceChild(Device)
+	, D3DDevice(Device->GetDevice())
+{
+}
+
 FD3D12ExplicitDescriptorHeapCache::~FD3D12ExplicitDescriptorHeapCache()
 {
 	check(NumAllocatedEntries == 0);
@@ -269,7 +275,7 @@ int32 FD3D12ExplicitDescriptorHeap::Allocate(uint32 InNumDescriptors)
 void FD3D12ExplicitDescriptorHeap::CopyDescriptors(int32 BaseIndex, const D3D12_CPU_DESCRIPTOR_HANDLE* InDescriptors, uint32 InNumDescriptors)
 {
 	D3D12_CPU_DESCRIPTOR_HANDLE DestDescriptor = GetDescriptorCPU(BaseIndex);
-	GetParentDevice()->GetDevice()->CopyDescriptors(1, &DestDescriptor, &InNumDescriptors, InNumDescriptors, InDescriptors, nullptr, Type);
+	D3DDevice->CopyDescriptors(1, &DestDescriptor, &InNumDescriptors, InNumDescriptors, InDescriptors, nullptr, Type);
 #if EXPLICIT_DESCRIPTOR_CACHE_FULL_COMPARE
 	for (uint32 i = 0; i < InNumDescriptors; ++i)
 	{
@@ -394,7 +400,20 @@ int32 FD3D12ExplicitDescriptorCache::Allocate(const D3D12_CPU_DESCRIPTOR_HANDLE*
 
 	FD3D12ExplicitDescriptorHeap& Heap = (Type == D3D12_DESCRIPTOR_HEAP_TYPE_CBV_SRV_UAV) ? ViewHeap : SamplerHeap;
 
-	int32 DescriptorTableBaseIndex = Heap.Allocate(NumDescriptors);
+	int32 DescriptorTableBaseIndex = INDEX_NONE; 
+	
+	if (Type == D3D12_DESCRIPTOR_HEAP_TYPE_CBV_SRV_UAV)
+	{
+		DescriptorTableBaseIndex = WorkerData[WorkerIndex].ReservedViewDescriptors.Allocate(NumDescriptors);
+		if (DescriptorTableBaseIndex == INDEX_NONE)
+		{
+			DescriptorTableBaseIndex = Heap.Allocate(NumDescriptors);
+		}
+	}
+	else
+	{
+		DescriptorTableBaseIndex = Heap.Allocate(NumDescriptors);
+	}
 
 	if (DescriptorTableBaseIndex == INDEX_NONE)
 	{
diff --git a/Engine/Source/Runtime/D3D12RHI/Private/D3D12ExplicitDescriptorCache.h b/Engine/Source/Runtime/D3D12RHI/Private/D3D12ExplicitDescriptorCache.h
--- a/Engine/Source/Runtime/D3D12RHI/Private/D3D12ExplicitDescriptorCache.h
+++ b/Engine/Source/Runtime/D3D12RHI/Private/D3D12ExplicitDescriptorCache.h
@@ -66,11 +66,7 @@ struct FD3D12ExplicitDescriptorHeap : public FD3D12DeviceChild
 {
 	UE_NONCOPYABLE(FD3D12ExplicitDescriptorHeap)
 
-	FD3D12ExplicitDescriptorHeap(FD3D12Device* Device)
-	: FD3D12DeviceChild(Device)
-	{
-	}
-
+	FD3D12ExplicitDescriptorHeap(FD3D12Device* Device);
 	~FD3D12ExplicitDescriptorHeap();
 
 	void Init(uint32 InMaxNumDescriptors, D3D12_DESCRIPTOR_HEAP_TYPE InType);
@@ -86,6 +82,9 @@ struct FD3D12ExplicitDescriptorHeap : public FD3D12DeviceChild
 	D3D12_CPU_DESCRIPTOR_HANDLE GetDescriptorCPU(uint32 Index) const;
 	D3D12_GPU_DESCRIPTOR_HANDLE GetDescriptorGPU(uint32 Index) const;
 
+	// Cache D3D device pointer, as it's frequently accessed on the hot path in CopyDescriptors
+	ID3D12Device* D3DDevice = nullptr;
+
 	D3D12_DESCRIPTOR_HEAP_TYPE Type = D3D12_DESCRIPTOR_HEAP_TYPE_NUM_TYPES;
 	ID3D12DescriptorHeap* D3D12Heap = nullptr;
 	uint32 MaxNumDescriptors = 0;
@@ -153,10 +152,47 @@ public:
 
 	using TDescriptorHashMap = Experimental::TSherwoodMap<uint64, int32, TIdentityHash<uint64>>;
 
+	struct FDescriptorSlotRange
+	{
+		FDescriptorSlotRange() = default;
+		FDescriptorSlotRange(int32 BaseIndex, int32 Count)
+		: Begin(BaseIndex)
+		, Cursor(BaseIndex)
+		, End(BaseIndex + Count)
+		{
+		}
+
+		int32 Begin = 0;
+		int32 Cursor = 0;
+		int32 End = 0;
+
+		int32 Allocate(int32 Count)
+		{
+			int32 Result = INDEX_NONE;
+			if (Cursor + Count <= End)
+			{
+				Result = Cursor;
+				Cursor += Count;
+			}
+			return Result;
+		}
+	};
+
+	void ReserveViewDescriptors(uint32 Count, uint32 WorkerIndex)
+	{
+		const int32 BaseIndex = ViewHeap.Allocate(Count);
+		if (BaseIndex != INDEX_NONE)
+		{
+			WorkerData[WorkerIndex].ReservedViewDescriptors = FDescriptorSlotRange(BaseIndex, Count);
+		}
+	}
+
 	struct alignas(PLATFORM_CACHE_LINE_SIZE) FWorkerThreadData
 	{
 		TDescriptorHashMap ViewDescriptorTableCache;
 		TDescriptorHashMap SamplerDescriptorTableCache;
+
+		FDescriptorSlotRange ReservedViewDescriptors;
 	};
 
 	TArray<FWorkerThreadData> WorkerData;
