Remove platform specific reference

diff --git a/Engine/Source/Runtime/D3D12RHI/Private/D3D12RHIPrivate.h b/Engine/Source/Runtime/D3D12RHI/Private/D3D12RHIPrivate.h
--- a/Engine/Source/Runtime/D3D12RHI/Private/D3D12RHIPrivate.h
+++ b/Engine/Source/Runtime/D3D12RHI/Private/D3D12RHIPrivate.h
@@ -412,6 +412,9 @@ public:
 	virtual void RHISubmitCommandLists(FRHISubmitCommandListsArgs&& Args) final override;
 	virtual void RHIProcessDeleteQueue() final override;
 
+	// Some RHIs can be suspended, most likely due to the HW going into suspend mode, this allows device dependent threads to wait until it's safe
+	virtual FRWLock* RHIGetSuspendedLock() { return nullptr; }
+
 	virtual void RHIRunOnQueue(ED3D12RHIRunOnQueueType QueueType, TFunction<void(ID3D12CommandQueue*)>&& CodeToRun, bool bWaitForSubmission) final override;
 
 	virtual IRHITransientResourceAllocator* RHICreateTransientResourceAllocator() override;
diff --git a/Engine/Source/Runtime/D3D12RHI/Private/D3D12Texture.cpp b/Engine/Source/Runtime/D3D12RHI/Private/D3D12Texture.cpp
--- a/Engine/Source/Runtime/D3D12RHI/Private/D3D12Texture.cpp
+++ b/Engine/Source/Runtime/D3D12RHI/Private/D3D12Texture.cpp
@@ -910,7 +910,7 @@ FTextureRHIRef FD3D12DynamicRHI::RHIAsyncCreateTexture2D(uint32 SizeX, uint32 Si
 	const D3D12_RESOURCE_STATES InitialState = D3D12_RESOURCE_STATE_COMMON;
 
 #if !PLATFORM_WINDOWS
-	FRWScopeLock ReadLock(GetID3D12XboxCommonDynamicRHI()->GetRHISuspendedLock(), SLT_ReadOnly);
+	FRWScopeLock ReadLock(*RHIGetSuspendedLock(), SLT_ReadOnly);
 #endif
 
 	FD3D12Adapter* Adapter = &GetAdapter();
