Add FD3D12Device::CopyDescriptors() wrapper for the common case of copying a single range of descriptors, with platform-specific implementations.

diff --git a/Engine/Source/Runtime/D3D12RHI/Private/D3D12DescriptorCache.cpp b/Engine/Source/Runtime/D3D12RHI/Private/D3D12DescriptorCache.cpp
--- a/Engine/Source/Runtime/D3D12RHI/Private/D3D12DescriptorCache.cpp
+++ b/Engine/Source/Runtime/D3D12RHI/Private/D3D12DescriptorCache.cpp
@@ -257,10 +257,7 @@ D3D12_GPU_DESCRIPTOR_HANDLE FD3D12DescriptorCache::BuildUAVTable(EShaderFrequenc
 	check((CurrentDirtySlotMask & SlotsNeededMask) == 0);	// Check all slots that needed to be set, were set.
 
 	// Gather the descriptors from the offline heaps to the online heap
-	GetParentDevice()->GetDevice()->CopyDescriptors(
-		1, &DestDescriptor, &SlotsNeeded,
-		SlotsNeeded, SrcDescriptors, nullptr /* sizes */,
-		D3D12_DESCRIPTOR_HEAP_TYPE_CBV_SRV_UAV);
+	GetParentDevice()->CopyDescriptors(DestDescriptor, SrcDescriptors, SlotsNeeded, D3D12_DESCRIPTOR_HEAP_TYPE_CBV_SRV_UAV);
 
 	return BindDescriptor;
 }
@@ -397,10 +394,7 @@ D3D12_GPU_DESCRIPTOR_HANDLE FD3D12DescriptorCache::BuildSamplerTable(EShaderFreq
 		}
 		FD3D12SamplerStateCache::CleanSlots(CurrentDirtySlotMask, SlotsNeeded);
 
-		GetParentDevice()->GetDevice()->CopyDescriptors(
-			1, &DestDescriptor, &SlotsNeeded,
-			SlotsNeeded, SrcDescriptors, nullptr /* sizes */,
-			FD3D12_DESCRIPTOR_HEAP_TYPE_SAMPLER);
+		GetParentDevice()->CopyDescriptors(DestDescriptor, SrcDescriptors, SlotsNeeded, FD3D12_DESCRIPTOR_HEAP_TYPE_SAMPLER);
 
 		// Remember the locations of the samplers in the sampler map
 		if (SlotsNeeded <= UE_ARRAY_COUNT(Desc.SamplerID))
@@ -480,8 +474,7 @@ D3D12_GPU_DESCRIPTOR_HANDLE FD3D12DescriptorCache::BuildSRVTable(EShaderFrequenc
 	}
 	FD3D12ShaderResourceViewCache::CleanSlots(CurrentDirtySlotMask, SlotsNeeded);
 
-	ID3D12Device* Device = GetParentDevice()->GetDevice();
-	Device->CopyDescriptors(1, &DestDescriptor, &SlotsNeeded, SlotsNeeded, SrcDescriptors, nullptr, D3D12_DESCRIPTOR_HEAP_TYPE_CBV_SRV_UAV);
+	GetParentDevice()->CopyDescriptors(DestDescriptor, SrcDescriptors, SlotsNeeded, D3D12_DESCRIPTOR_HEAP_TYPE_CBV_SRV_UAV);
 
 	check((CurrentDirtySlotMask & SlotsNeededMask) == 0);	// Check all slots that needed to be set, were set.
 
@@ -1076,10 +1069,7 @@ void FD3D12GlobalOnlineSamplerHeap::ConsolidateUniqueSamplerTables(TArrayView<FD
 
 				D3D12_CPU_DESCRIPTOR_HANDLE DestDescriptor = GetCPUSlotHandle(HeapSlot);
 
-				GetParentDevice()->GetDevice()->CopyDescriptors(
-					1, &DestDescriptor, &Table.Key.Count,
-					Table.Key.Count, Table.CPUTable, nullptr /* sizes */,
-					FD3D12_DESCRIPTOR_HEAP_TYPE_SAMPLER);
+				GetParentDevice()->CopyDescriptors(DestDescriptor, Table.CPUTable, Table.Key.Count, FD3D12_DESCRIPTOR_HEAP_TYPE_SAMPLER);
 
 				Table.GPUHandle = GetGPUSlotHandle(HeapSlot);
 				UniqueDescriptorTables->Add(Table);
diff --git a/Engine/Source/Runtime/D3D12RHI/Private/D3D12Device.h b/Engine/Source/Runtime/D3D12RHI/Private/D3D12Device.h
--- a/Engine/Source/Runtime/D3D12RHI/Private/D3D12Device.h
+++ b/Engine/Source/Runtime/D3D12RHI/Private/D3D12Device.h
@@ -202,6 +202,13 @@ public:
 	D3D12_RESOURCE_ALLOCATION_INFO GetResourceAllocationInfoUncached(const FD3D12ResourceDesc& InDesc);
 	D3D12_RESOURCE_ALLOCATION_INFO GetResourceAllocationInfo(const FD3D12ResourceDesc& InDesc);
 
+	// Specialized wrapper of ID3D12Device::CopyDescriptors for a common case of a single descriptor range. Similar to CopyDescriptorsSimple(), except source is provided as an array.
+	static void CopyDescriptors(ID3D12Device* D3DDevice, D3D12_CPU_DESCRIPTOR_HANDLE Destination, const D3D12_CPU_DESCRIPTOR_HANDLE* Source, uint32 NumSourceDescriptors, D3D12_DESCRIPTOR_HEAP_TYPE Type);
+	void CopyDescriptors(D3D12_CPU_DESCRIPTOR_HANDLE Destination, const D3D12_CPU_DESCRIPTOR_HANDLE* Source, uint32 NumSourceDescriptors, D3D12_DESCRIPTOR_HEAP_TYPE Type)
+	{
+		CopyDescriptors(GetDevice(), Destination, Source, NumSourceDescriptors, Type);
+	}
+
 	void									  InitExplicitDescriptorHeap();
 	FD3D12ExplicitDescriptorHeapCache*		  GetExplicitDescriptorHeapCache() { return ExplicitDescriptorHeapCache; }
 
diff --git a/Engine/Source/Runtime/D3D12RHI/Private/D3D12ExplicitDescriptorCache.cpp b/Engine/Source/Runtime/D3D12RHI/Private/D3D12ExplicitDescriptorCache.cpp
--- a/Engine/Source/Runtime/D3D12RHI/Private/D3D12ExplicitDescriptorCache.cpp
+++ b/Engine/Source/Runtime/D3D12RHI/Private/D3D12ExplicitDescriptorCache.cpp
@@ -275,7 +275,8 @@ int32 FD3D12ExplicitDescriptorHeap::Allocate(uint32 InNumDescriptors)
 void FD3D12ExplicitDescriptorHeap::CopyDescriptors(int32 BaseIndex, const D3D12_CPU_DESCRIPTOR_HANDLE* InDescriptors, uint32 InNumDescriptors)
 {
 	D3D12_CPU_DESCRIPTOR_HANDLE DestDescriptor = GetDescriptorCPU(BaseIndex);
-	D3DDevice->CopyDescriptors(1, &DestDescriptor, &InNumDescriptors, InNumDescriptors, InDescriptors, nullptr, Type);
+	FD3D12Device::CopyDescriptors(D3DDevice, DestDescriptor, InDescriptors, InNumDescriptors, Type);
+
 #if EXPLICIT_DESCRIPTOR_CACHE_FULL_COMPARE
 	for (uint32 i = 0; i < InNumDescriptors; ++i)
 	{
diff --git a/Engine/Source/Runtime/D3D12RHI/Private/Windows/WindowsD3D12Device.cpp b/Engine/Source/Runtime/D3D12RHI/Private/Windows/WindowsD3D12Device.cpp
--- a/Engine/Source/Runtime/D3D12RHI/Private/Windows/WindowsD3D12Device.cpp
+++ b/Engine/Source/Runtime/D3D12RHI/Private/Windows/WindowsD3D12Device.cpp
@@ -1853,6 +1853,11 @@ void FD3D12Device::CreateSamplerInternal(const D3D12_SAMPLER_DESC& Desc, D3D12_C
 	GetDevice()->CreateSampler(&Desc, Descriptor);
 }
 
+void FD3D12Device::CopyDescriptors(ID3D12Device* D3DDevice, D3D12_CPU_DESCRIPTOR_HANDLE Destination, const D3D12_CPU_DESCRIPTOR_HANDLE* Source, uint32 NumSourceDescriptors, D3D12_DESCRIPTOR_HEAP_TYPE Type)
+{
+	D3DDevice->CopyDescriptors(1, &Destination, &NumSourceDescriptors, NumSourceDescriptors, Source, nullptr, Type);
+}
+
 #if D3D12_RHI_RAYTRACING
 void FD3D12Device::GetRaytracingAccelerationStructurePrebuildInfo(
 	const D3D12_BUILD_RAYTRACING_ACCELERATION_STRUCTURE_INPUTS* pDesc,
