Add stat counters for the number of executed command lists per frame

diff --git a/Engine/Source/Runtime/D3D12RHI/Private/D3D12Stats.h b/Engine/Source/Runtime/D3D12RHI/Private/D3D12Stats.h
--- a/Engine/Source/Runtime/D3D12RHI/Private/D3D12Stats.h
+++ b/Engine/Source/Runtime/D3D12RHI/Private/D3D12Stats.h
@@ -25,6 +25,9 @@ DECLARE_DWORD_ACCUMULATOR_STAT_EXTERN(TEXT("Num command lists (3D, Compute, Copy
 DECLARE_DWORD_ACCUMULATOR_STAT_EXTERN(TEXT("Num Query Heaps"), STAT_D3D12NumQueryHeaps, STATGROUP_D3D12RHI, );
 DECLARE_DWORD_ACCUMULATOR_STAT_EXTERN(TEXT("Num pipeline state objects (PSOs)"), STAT_D3D12NumPSOs, STATGROUP_D3D12RHI, );
 
+DECLARE_DWORD_COUNTER_STAT_EXTERN(TEXT("Executed Command Lists"), STAT_D3D12ExecutedCommandLists, STATGROUP_D3D12RHI, );
+DECLARE_DWORD_COUNTER_STAT_EXTERN(TEXT("Executed Command List Batches"), STAT_D3D12ExecutedCommandListBatches, STATGROUP_D3D12RHI, );
+
 DECLARE_DWORD_COUNTER_STAT_EXTERN(TEXT("Textures Allocated"), STAT_D3D12TexturesAllocated, STATGROUP_D3D12RHI, );
 DECLARE_DWORD_COUNTER_STAT_EXTERN(TEXT("Textures Released"), STAT_D3D12TexturesReleased, STATGROUP_D3D12RHI, );
 DECLARE_CYCLE_STAT_EXTERN(TEXT("CreateTexture time"), STAT_D3D12CreateTextureTime, STATGROUP_D3D12RHI, );
diff --git a/Engine/Source/Runtime/D3D12RHI/Private/D3D12Submission.cpp b/Engine/Source/Runtime/D3D12RHI/Private/D3D12Submission.cpp
--- a/Engine/Source/Runtime/D3D12RHI/Private/D3D12Submission.cpp
+++ b/Engine/Source/Runtime/D3D12RHI/Private/D3D12Submission.cpp
@@ -818,6 +818,9 @@ void FD3D12DynamicRHI::FlushBatchedPayloads(FD3D12Queue::FPayloadArray& Payloads
 				DispatchNum = Index;
 			}
 
+			INC_DWORD_STAT(STAT_D3D12ExecutedCommandListBatches);
+			INC_DWORD_STAT_BY(STAT_D3D12ExecutedCommandLists, DispatchNum);
+
 			Queue.ExecuteCommandLists(
 				MakeArrayView<ID3D12CommandList*>(&D3DCommandLists[Offset], DispatchNum)
 #if ENABLE_RESIDENCY_MANAGEMENT
diff --git a/Engine/Source/Runtime/D3D12RHI/Private/D3D12Util.cpp b/Engine/Source/Runtime/D3D12RHI/Private/D3D12Util.cpp
--- a/Engine/Source/Runtime/D3D12RHI/Private/D3D12Util.cpp
+++ b/Engine/Source/Runtime/D3D12RHI/Private/D3D12Util.cpp
@@ -1843,6 +1843,8 @@ DEFINE_STAT(STAT_D3D12NumCommandAllocators);
 DEFINE_STAT(STAT_D3D12NumCommandLists);
 DEFINE_STAT(STAT_D3D12NumQueryHeaps);
 DEFINE_STAT(STAT_D3D12NumPSOs);
+DEFINE_STAT(STAT_D3D12ExecutedCommandLists);
+DEFINE_STAT(STAT_D3D12ExecutedCommandListBatches);
 
 DEFINE_STAT(STAT_D3D12TexturesAllocated);
 DEFINE_STAT(STAT_D3D12TexturesReleased);
