Aftermath shader association and general improvements
- Added general Aftermath hash to (DXBC) shader hash association
- Added "late" associated (default), which associates the relevant hashes after a GPU crash has occurred. Association after pipeline creation can reach many milliseconds, sometimes +100ms.
- Added pipeline frame counters, avoid late association on pipelines that havent been used in N frames. Late association goes from seconds to <100ms.
- Added GetPipelineStates to the pipeline cache, interrupts the pipeline caches to retrieve all states
- Added DXBC utility to get the shader hash
- Added generic file paths to crash dump handling, copies the files to the crash dump directory for uploading. Also avoids associating minidumps by the latest modified, which may have false positives

diff --git a/Engine/Source/Programs/CrashReportClient/Private/CrashReportClientApp.cpp b/Engine/Source/Programs/CrashReportClient/Private/CrashReportClientApp.cpp
--- a/Engine/Source/Programs/CrashReportClient/Private/CrashReportClientApp.cpp
+++ b/Engine/Source/Programs/CrashReportClient/Private/CrashReportClientApp.cpp
@@ -515,6 +515,7 @@ FPlatformErrorReport CollectErrorReport(FRecoveryService* RecoveryService, uint3
 		SharedCrashContext.EnabledPluginsNum > 0 ? &SharedCrashContext.DynamicData[SharedCrashContext.EnabledPluginsOffset] : nullptr,
 		SharedCrashContext.EngineDataNum > 0 ? &SharedCrashContext.DynamicData[SharedCrashContext.EngineDataOffset] : nullptr,
 		SharedCrashContext.GameDataNum > 0 ? &SharedCrashContext.DynamicData[SharedCrashContext.GameDataOffset] : nullptr,
+		SharedCrashContext.FilePathsNum > 0 ? &SharedCrashContext.DynamicData[SharedCrashContext.FilePathsOffset] : nullptr,
 		&SharedCrashContext.GPUBreadcrumbs
 	);
 	// Next create a crash context for the crashed process.
diff --git a/Engine/Source/Runtime/Core/Private/GenericPlatform/GenericPlatformCrashContext.cpp b/Engine/Source/Runtime/Core/Private/GenericPlatform/GenericPlatformCrashContext.cpp
--- a/Engine/Source/Runtime/Core/Private/GenericPlatform/GenericPlatformCrashContext.cpp
+++ b/Engine/Source/Runtime/Core/Private/GenericPlatform/GenericPlatformCrashContext.cpp
@@ -306,12 +306,15 @@ FAdditionalCrashContextDelegate FGenericCrashContext::AdditionalCrashContextDele
 
 namespace NCached
 {
+	static FCriticalSection CriticalSection;
+	
 	static FSessionContext Session;
 	static FUserSettingsContext UserSettings;
 	static TArray<FString> EnabledPluginsList;
 	static TMap<FString, FString> EngineData;
 	static TMap<FString, FString> GameData;
 	static TOptional<FGPUBreadcrumbCrashData> GPUBreadcrumbs;
+	static TArray<FString> FilePaths;
 
 	template <size_t CharCount, typename CharType>
 	void Set(CharType(&Dest)[CharCount], const CharType* pSrc)
@@ -531,7 +534,7 @@ void FGenericCrashContext::Initialize()
 const TCHAR* CR_PAIR_DELIM = TEXT("\x01");
 const TCHAR* CR_PAIR_EQ = TEXT("\x02");
 
-void FGenericCrashContext::InitializeFromContext(const FSessionContext& Session, const TCHAR* EnabledPluginsStr, const TCHAR* EngineDataStr, const TCHAR* GameDataStr, const FGPUBreadcrumbsSharedContext* GPUBreadcrumbs)
+void FGenericCrashContext::InitializeFromContext(const FSessionContext& Session, const TCHAR* EnabledPluginsStr, const TCHAR* EngineDataStr, const TCHAR* GameDataStr, const TCHAR* FilePathData, const FGPUBreadcrumbsSharedContext* GPUBreadcrumbs)
 {
 	static const TCHAR* TokenDelim[] = { CR_PAIR_DELIM, CR_PAIR_EQ };
 
@@ -581,6 +584,17 @@ void FGenericCrashContext::InitializeFromContext(const FSessionContext& Session,
 		}
 	}
 
+	if (FilePathData)
+	{
+		TArray<FString> Tokens;
+		FString(FilePathData).ParseIntoArray(Tokens, TokenDelim, 2, true);
+
+		for (FString& Token : Tokens)
+		{
+			NCached::FilePaths.Add(Token);
+		}
+	}
+
 	if (GPUBreadcrumbs && GPUBreadcrumbs->NumQueues > 0)
 	{
 		SetGPUBreadcrumbs(GPUBreadcrumbsFromSharedContext(*GPUBreadcrumbs));
@@ -651,6 +665,14 @@ void FGenericCrashContext::CopySharedCrashContext(FSharedCrashContext& Dst)
 	}
 	DynamicDataPtr += FCString::Strlen(DynamicDataPtr) + 1;
 
+	Dst.FilePathsOffset = (uint32)(DynamicDataPtr - DynamicDataStart);
+	Dst.FilePathsNum = NCached::FilePaths.Num();
+	for (const FString& FilePath : NCached::FilePaths)
+	{
+		FCString::StrncatTruncateDest(DynamicDataPtr, CR_DYNAMIC_BUFFER_REMAIN, *FilePath);
+		FCString::StrncatTruncateDest(DynamicDataPtr, CR_DYNAMIC_BUFFER_REMAIN, CR_PAIR_DELIM);
+	}
+
 	#undef CR_DYNAMIC_BUFFER_REMAIN
 
 	// Copy GPU breadcrumbs.
@@ -1541,6 +1563,12 @@ void FGenericCrashContext::AddPlugin(const FString& PluginDesc)
 	NCached::EnabledPluginsList.Add(PluginDesc);
 }
 
+void FGenericCrashContext::AddFile(const FString& AbsoluteSourcePath)
+{
+	FScopeLock Lock(&NCached::CriticalSection);
+	NCached::FilePaths.Add(AbsoluteSourcePath);
+}
+
 FString FGenericCrashContext::DumpLog(const FString& CrashFolderAbsolute)
 {
 	// Copy log
@@ -1718,7 +1746,19 @@ void FGenericCrashContext::CopyPlatformSpecificFiles(const TCHAR* OutputDirector
 		IFileManager::Get().Copy(*CrashConfigDstAbsolute, CrashConfigSrcPath);	// best effort, so don't care about result: couldn't copy -> tough, no config
 	}
 
-	
+	for (const FString& AbsoluteSource : NCached::FilePaths)
+	{
+		if (IFileManager::Get().FileExists(*AbsoluteSource))
+		{
+			const FString AbsoluteDest = FPaths::Combine(OutputDirectory, FPaths::GetCleanFilename(AbsoluteSource));
+
+			// Try to move file, may fail due to any number of reasons
+			if (!IFileManager::Get().Move(*AbsoluteDest, *AbsoluteSource))
+			{
+				UE_LOG(LogCrashContext, Error, TEXT("Error moving crash context file '%s' to '%s'"), *AbsoluteSource, *AbsoluteDest);
+			}
+		}
+	}
 }
 
 #if WITH_ADDITIONAL_CRASH_CONTEXTS
diff --git a/Engine/Source/Runtime/Core/Private/Windows/WindowsPlatformCrashContext.cpp b/Engine/Source/Runtime/Core/Private/Windows/WindowsPlatformCrashContext.cpp
--- a/Engine/Source/Runtime/Core/Private/Windows/WindowsPlatformCrashContext.cpp
+++ b/Engine/Source/Runtime/Core/Private/Windows/WindowsPlatformCrashContext.cpp
@@ -201,8 +201,6 @@ struct FAssertInfo
 	}
 };
 
-static constexpr TCHAR const NvAftermathDumpExtension[] = TEXT(".nv-gpudmp");
-
 namespace UE::Core::Private
 {
 	static TAutoConsoleVariable<bool> CVarForceCrashReportDialogOff(
@@ -381,37 +379,6 @@ void FWindowsPlatformCrashContext::CopyPlatformSpecificFiles(const TCHAR* Output
 		static_cast<void>(IFileManager::Get().Copy(*CrashVideoDstAbsolute, *CrashVideoPath));	// best effort, so don't care about result: couldn't copy -> tough, no video
 	}
 
-	// Find the newest GPU crash dump file
-	// Best effort, so don't care about result: couldn't copy -> tough, no GPU crash dump
-	TArray<FString> GPUDumpFiles;
-	IFileManager::Get().FindFiles(GPUDumpFiles, *FPaths::ProjectLogDir(), NvAftermathDumpExtension);
-	for (FString& GPUDumpFilename : GPUDumpFiles)
-	{
-		GPUDumpFilename = FPaths::Combine(*FPaths::ProjectLogDir(), *GPUDumpFilename);
-	}
-
-	if (GPUDumpFiles.Num())
-	{
-		GPUDumpFiles.Sort([](FString LHS, FString RHS)
-		{
-			double LHSAge = IFileManager::Get().GetFileAgeSeconds(*LHS);
-			double RHSAge = IFileManager::Get().GetFileAgeSeconds(*RHS);
-
-			return LHSAge < RHSAge;
-		});
-		FString SelectedGPUCrashDump = GPUDumpFiles[0];
-
-		if (IFileManager::Get().FileExists(*SelectedGPUCrashDump))
-		{
-			FString GPUMiniDumpFilename = FPaths::GetCleanFilename(SelectedGPUCrashDump);
-			const FString GPUMiniDumpDstAbsolute = FPaths::Combine(OutputDirectory, *GPUMiniDumpFilename);
-			if (!IFileManager::Get().Move(*GPUMiniDumpDstAbsolute, *SelectedGPUCrashDump))
-			{
-				UE_LOG(LogWindows, Error, TEXT("Error moving GPU crash dump file %s to output crash directory"), *GPUMiniDumpFilename);
-			}
-		}
-	}
-
 	CopyOptionalAttachments(OutputDirectory);
 }
 
diff --git a/Engine/Source/Runtime/Core/Public/GenericPlatform/GenericPlatformCrashContext.h b/Engine/Source/Runtime/Core/Public/GenericPlatform/GenericPlatformCrashContext.h
--- a/Engine/Source/Runtime/Core/Public/GenericPlatform/GenericPlatformCrashContext.h
+++ b/Engine/Source/Runtime/Core/Public/GenericPlatform/GenericPlatformCrashContext.h
@@ -273,6 +273,9 @@ struct FSharedCrashContext
 	// Count and offset into dynamic buffer to comma separated key=value data for  game information
 	uint32					GameDataNum;
 	uint32					GameDataOffset;
+	// Count and offset into dynamic buffer to comma separated file list
+	uint32					FilePathsNum;
+	uint32					FilePathsOffset;
 	// Fixed size dynamic buffer
 	TCHAR					DynamicData[CR_MAX_DYNAMIC_BUFFER_CHARS];
 
@@ -457,7 +460,7 @@ public:
 	CORE_API static void Initialize();
 
 	/** Initialized crash context, using a crash context (e.g. shared from another process). */
-	CORE_API static void InitializeFromContext(const FSessionContext& Context, const TCHAR* EnabledPlugins, const TCHAR* EngineData, const TCHAR* GameData, const FGPUBreadcrumbsSharedContext* GPUBreadcrumbs);
+	CORE_API static void InitializeFromContext(const FSessionContext& Context, const TCHAR* EnabledPlugins, const TCHAR* EngineData, const TCHAR* GameData, const TCHAR* FilePathData, const FGPUBreadcrumbsSharedContext* GPUBreadcrumbs);
 
 	/** Get the current cached session context */
 	CORE_API static const FSessionContext& GetCachedSessionContext();
@@ -632,6 +635,9 @@ public:
 	/** Adds a plugin descriptor string to the enabled plugins list in the crash context */
 	CORE_API static void AddPlugin(const FString& PluginDesc);
 
+	/** Adds a file to bundle with the crash report */
+	CORE_API static void AddFile(const FString& AbsoluteSourcePath);
+
 	/** Flushes the logs. In the case of in memory logs is used on this configuration, dumps them to file. Returns the name of the file */
 	CORE_API static FString DumpLog(const FString& CrashFolderAbsolute);
 
diff --git a/Engine/Source/Runtime/D3D12RHI/Private/D3D12Adapter.cpp b/Engine/Source/Runtime/D3D12RHI/Private/D3D12Adapter.cpp
--- a/Engine/Source/Runtime/D3D12RHI/Private/D3D12Adapter.cpp
+++ b/Engine/Source/Runtime/D3D12RHI/Private/D3D12Adapter.cpp
@@ -396,8 +396,9 @@ void FD3D12Adapter::CreateRootDevice(bool bWithDebug)
 	if (IsRHIDeviceNVIDIA())
 	{
 		UE::RHICore::Nvidia::Aftermath::InitializeBeforeDeviceCreation();
+		UE::RHICore::Nvidia::Aftermath::SetLateShaderAssociateCallback(UE::RHICore::Nvidia::Aftermath::D3D12::CreateShaderAssociations);
 	}
-#endif
+#endif // NV_AFTERMATH
 
 	// Setup DRED if requested
 	bool bDRED = false;
diff --git a/Engine/Source/Runtime/D3D12RHI/Private/D3D12CommandContext.cpp b/Engine/Source/Runtime/D3D12RHI/Private/D3D12CommandContext.cpp
--- a/Engine/Source/Runtime/D3D12RHI/Private/D3D12CommandContext.cpp
+++ b/Engine/Source/Runtime/D3D12RHI/Private/D3D12CommandContext.cpp
@@ -495,6 +495,11 @@ void FD3D12ContextCommon::Finalize(TArray<FD3D12Payload*>& OutPayloads)
 	OutPayloads.Append(MoveTemp(Payloads));
 }
 
+uint32 FD3D12CommandContext::GetFrameFenceCounter() const
+{
+	return GetParentDevice()->GetParentAdapter()->GetFrameFence().GetNextFenceToSignal();
+}
+
 void FD3D12CommandContext::Finalize(TArray<FD3D12Payload*>& OutPayloads)
 {
 #if PLATFORM_SUPPORTS_BINDLESS_RENDERING
diff --git a/Engine/Source/Runtime/D3D12RHI/Private/D3D12CommandContext.h b/Engine/Source/Runtime/D3D12RHI/Private/D3D12CommandContext.h
--- a/Engine/Source/Runtime/D3D12RHI/Private/D3D12CommandContext.h
+++ b/Engine/Source/Runtime/D3D12RHI/Private/D3D12CommandContext.h
@@ -780,6 +780,8 @@ public:
 		return RetrieveTexture(Texture, GetGPUIndex());
 	}
 
+	uint32 GetFrameFenceCounter() const;
+
 	uint32 GetGPUIndex() const { return GPUMask.ToIndex(); }
 
 	virtual void RHISetGPUMask(FRHIGPUMask InGPUMask) final override
diff --git a/Engine/Source/Runtime/D3D12RHI/Private/D3D12Commands.cpp b/Engine/Source/Runtime/D3D12RHI/Private/D3D12Commands.cpp
--- a/Engine/Source/Runtime/D3D12RHI/Private/D3D12Commands.cpp
+++ b/Engine/Source/Runtime/D3D12RHI/Private/D3D12Commands.cpp
@@ -992,6 +992,8 @@ void FD3D12CommandContext::RHISetGraphicsPipelineState(FRHIGraphicsPipelineState
 		}
 	}
 
+	GraphicsPipelineState->FrameCounter.Set(GetFrameFenceCounter());
+
 	StateCache.SetGraphicsPipelineState(GraphicsPipelineState);
 	StateCache.SetStencilRef(StencilRef);
 
@@ -1016,6 +1018,8 @@ void FD3D12CommandContext::RHISetComputePipelineState(FRHIComputePipelineState*
 	StageConstantBuffers[SF_Compute].Reset();
 	bDiscardSharedComputeConstants = true;
 
+	ComputePipelineState->FrameCounter.Set(GetFrameFenceCounter());
+
 	StateCache.SetComputePipelineState(ComputePipelineState);
 
 	ApplyStaticUniformBuffersOnContext(*this, FD3D12DynamicRHI::ResourceCast(ComputePipelineState->GetComputeShader()));
diff --git a/Engine/Source/Runtime/D3D12RHI/Private/D3D12NvidiaExtensions.cpp b/Engine/Source/Runtime/D3D12RHI/Private/D3D12NvidiaExtensions.cpp
--- a/Engine/Source/Runtime/D3D12RHI/Private/D3D12NvidiaExtensions.cpp
+++ b/Engine/Source/Runtime/D3D12RHI/Private/D3D12NvidiaExtensions.cpp
@@ -1,7 +1,12 @@
 // Copyright Epic Games, Inc. All Rights Reserved.
 
 #include "D3D12NvidiaExtensions.h"
+#include "D3D12PipelineState.h"
+#include "D3D12RayTracing.h"
+#include "D3D12WorkGraph.h"
 #include "D3D12RHICommon.h"
+#include "PipelineStateCache.h"
+#include "Async/ParallelFor.h"
 
 #if NV_AFTERMATH
 
@@ -19,6 +24,20 @@ THIRD_PARTY_INCLUDES_END
 	#include "Windows/HideWindowsPlatformTypes.h"
 #endif
 
+static TAutoConsoleVariable<float> CVarAftermathLateShaderAssociationsTimeLimit(
+	TEXT("r.GPUCrashDebugging.Aftermath.D3D12.LateShaderAssociations.TimeLimit"),
+	7.5f,
+	TEXT("Time limit (s) before the late associations are stopped\n"),
+	ECVF_Default
+);
+
+static TAutoConsoleVariable<int32> CVarAftermathLateShaderAssociationFrameLimit(
+	TEXT("r.GPUCrashDebugging.Aftermath.D3D12.LateShaderAssociations.FrameLimit"),
+	10,
+	TEXT("Max last since used number of frames to consider when collecting pipelines\n"),
+	ECVF_Default
+);
+
 namespace UE::RHICore::Nvidia::Aftermath::D3D12
 {
 	void InitializeDevice(ID3D12Device* RootDevice)
@@ -29,6 +48,191 @@ namespace UE::RHICore::Nvidia::Aftermath::D3D12
 		});
 	}
 
+	template<typename T>
+	void RegisterShader(T* Value)
+	{
+		RegisterShaderBinary(Value->Code.GetData(), Value->Code.Num());
+	}
+	
+	template<typename T>
+	void ConditionallyAddShader(T* Value, TArray<FRHIShader*>& Shaders, TSet<FRHIShader*>& ShaderSet)
+	{
+		if (Value && !ShaderSet.Contains(Value))
+		{
+			Shaders.Add(Value);
+			ShaderSet.Add(Value);
+		}
+	}
+
+	bool IsPipelineUnused(FD3D12Adapter* Adapter, const D3D12ResourceFrameCounter& Counter, uint32_t Threshold)
+	{
+		const uint32 AdapterFrameIndex = Adapter->GetFrameFence().GetNextFenceToSignal();
+		const uint32 PipelineFrameIndex = Counter.Get();
+
+		// Check in case that it's ahead of the adapter, just in case it entered a strange state
+		return (PipelineFrameIndex <= AdapterFrameIndex) && (AdapterFrameIndex - PipelineFrameIndex) > Threshold;
+	}
+
+	void CreateShaderAssociations()
+	{
+		float TimeLimit = CVarAftermathLateShaderAssociationsTimeLimit->GetFloat();
+		
+		uint64 CycleStart = FPlatformTime::Cycles64();
+
+		uint32 FrameThreshold = CVarAftermathLateShaderAssociationFrameLimit->GetInt();
+
+		UE_LOG(LogD3D12RHI, Log, TEXT("Starting late shader associations..."));
+
+		TArray<FRHIShader*> Shaders;
+		TSet<FRHIShader*>   ShaderSet;
+
+		uint32_t IgnoredPipelines = 0;
+
+		TArray<TRefCountPtr<FRHIResource>> PipelineResources;
+		PipelineStateCache::GetPipelineStates(PipelineResources);
+
+		// Deduplicate shaders, Aftermath hashes are not local to the parent pipeline
+		for (FRHIResource* Resource : PipelineResources)
+		{
+			switch (Resource->GetType())
+			{
+				default:
+					checkNoEntry();
+				case RRT_GraphicsPipelineState:
+				{
+					auto* Pipeline = static_cast<FD3D12GraphicsPipelineState*>(Resource);
+						
+					if (IsPipelineUnused(Pipeline->PipelineState->GetParentAdapter(), Pipeline->FrameCounter, FrameThreshold))
+					{
+						IgnoredPipelines++;
+						continue;
+					}
+						
+					ConditionallyAddShader(static_cast<FD3D12VertexShader*>(Pipeline->PipelineStateInitializer.BoundShaderState.GetVertexShader()), Shaders, ShaderSet);
+					ConditionallyAddShader(static_cast<FD3D12GeometryShader*>(Pipeline->PipelineStateInitializer.BoundShaderState.GetGeometryShader()), Shaders, ShaderSet);
+					ConditionallyAddShader(static_cast<FD3D12AmplificationShader*>(Pipeline->PipelineStateInitializer.BoundShaderState.GetAmplificationShader()), Shaders, ShaderSet);
+					ConditionallyAddShader(static_cast<FD3D12MeshShader*>(Pipeline->PipelineStateInitializer.BoundShaderState.GetMeshShader()), Shaders, ShaderSet);
+					ConditionallyAddShader(static_cast<FD3D12PixelShader*>(Pipeline->PipelineStateInitializer.BoundShaderState.GetPixelShader()), Shaders, ShaderSet);
+					break;
+				}
+				case RRT_ComputePipelineState:
+				{
+					auto* Pipeline = static_cast<FD3D12ComputePipelineState*>(Resource);
+					
+					if (IsPipelineUnused(Pipeline->PipelineState->GetParentAdapter(), Pipeline->FrameCounter, FrameThreshold))
+					{
+						IgnoredPipelines++;
+						continue;
+					}
+						
+					ConditionallyAddShader(static_cast<FD3D12ComputeShader*>(Pipeline->GetComputeShader()), Shaders, ShaderSet);
+					break;
+				}
+				case RRT_RayTracingPipelineState:
+				{
+					auto* Pipeline = static_cast<FD3D12RayTracingPipelineState*>(Resource);
+						
+					if (IsPipelineUnused(Pipeline->Device->GetParentAdapter(), Pipeline->FrameCounter, FrameThreshold))
+					{
+						IgnoredPipelines++;
+						continue;
+					}
+						
+					for (const TRefCountPtr<FD3D12RayTracingShader>& Shader : Pipeline->RayGenShaders.Shaders)
+					{
+						ConditionallyAddShader(Shader.GetReference(), Shaders, ShaderSet);
+					}
+					for (const TRefCountPtr<FD3D12RayTracingShader>& Shader : Pipeline->CallableShaders.Shaders)
+					{
+						ConditionallyAddShader(Shader.GetReference(), Shaders, ShaderSet);
+					}
+					for (const TRefCountPtr<FD3D12RayTracingShader>& Shader : Pipeline->HitGroupShaders.Shaders)
+					{
+						ConditionallyAddShader(Shader.GetReference(), Shaders, ShaderSet);
+					}
+					for (const TRefCountPtr<FD3D12RayTracingShader>& Shader : Pipeline->MissShaders.Shaders)
+					{
+						ConditionallyAddShader(Shader.GetReference(), Shaders, ShaderSet);
+					}
+					break;
+				}
+				case RRT_WorkGraphPipelineState:
+				{
+					auto* Pipeline = static_cast<FD3D12WorkGraphPipelineState*>(Resource);
+						
+					if (IsPipelineUnused(Pipeline->Device->GetParentAdapter(), Pipeline->FrameCounter, FrameThreshold))
+					{
+						IgnoredPipelines++;
+						continue;
+					}
+						
+					for (FRHIShader* Shader : Pipeline->Shaders)
+					{
+						auto* ShaderTyped = static_cast<FD3D12RayTracingShader*>(Shader);
+						ConditionallyAddShader(ShaderTyped, Shaders, ShaderSet);
+					}
+					break;
+				}
+			}
+		}
+
+		UE_LOG(LogD3D12RHI, Log, TEXT("Late shader associations ignored %u pipelines based on frame fences"), IgnoredPipelines);
+
+		// Parallelize as much as possible to avoid timeouts
+		ParallelFor(Shaders.Num(), [CycleStart, TimeLimit, &Shaders](int32 Index)
+		{
+			// Aftermath handling is time constrained, if we hit the limit just stop
+			float Elapsed = FPlatformTime::ToSeconds64(FPlatformTime::Cycles64() - CycleStart);
+			if (Elapsed >= TimeLimit)
+			{
+				UE_CALL_ONCE([Elapsed]
+				{
+					UE_LOG(LogD3D12RHI, Warning, TEXT("Late shader associations timed out at %f.5ms"), Elapsed);
+				});
+				return;
+			}
+
+			FRHIShader* Shader = Shaders[Index];
+			switch (Shader->GetFrequency())
+			{
+				default:
+					checkNoEntry();
+					break;
+				case SF_Vertex:
+					RegisterShader(static_cast<FD3D12VertexShader*>(Shader));
+					break;
+				case SF_Amplification:
+					RegisterShader(static_cast<FD3D12AmplificationShader*>(Shader));
+					break;
+				case SF_Mesh:
+					RegisterShader(static_cast<FD3D12MeshShader*>(Shader));
+					break;
+				case SF_Geometry:
+					RegisterShader(static_cast<FD3D12GeometryShader*>(Shader));
+					break;
+				case SF_Pixel:
+					RegisterShader(static_cast<FD3D12PixelShader*>(Shader));
+					break;
+				case SF_Compute:
+					RegisterShader(static_cast<FD3D12ComputeShader*>(Shader));
+					break;
+				case SF_RayGen:
+				case SF_RayCallable:
+				case SF_RayHitGroup:
+				case SF_RayMiss:
+					RegisterShader(static_cast<FD3D12RayTracingShader*>(Shader));
+					break;
+				case SF_WorkGraphRoot:
+				case SF_WorkGraphComputeNode:
+					RegisterShader(static_cast<FD3D12WorkGraphShader*>(Shader));
+					break;
+			}
+		});
+		
+		double TimeMS = FPlatformTime::ToMilliseconds64(FPlatformTime::Cycles64() - CycleStart);
+		UE_LOG(LogD3D12RHI, Log, TEXT("Created late shader associations, took %f.5ms"), TimeMS);
+	}
+
 	FCommandList RegisterCommandList(ID3D12CommandList* D3DCommandList)
 	{
 		FCommandList Handle{};
diff --git a/Engine/Source/Runtime/D3D12RHI/Private/D3D12NvidiaExtensions.h b/Engine/Source/Runtime/D3D12RHI/Private/D3D12NvidiaExtensions.h
--- a/Engine/Source/Runtime/D3D12RHI/Private/D3D12NvidiaExtensions.h
+++ b/Engine/Source/Runtime/D3D12RHI/Private/D3D12NvidiaExtensions.h
@@ -31,6 +31,7 @@
 		typedef struct GFSDK_Aftermath_ResourceHandle__* FResource;
 
 		void InitializeDevice(ID3D12Device* RootDevice);
+		void CreateShaderAssociations();
 
 		FCommandList RegisterCommandList(ID3D12CommandList* D3DCommandList);
 		void UnregisterCommandList(FCommandList CommandList);
diff --git a/Engine/Source/Runtime/D3D12RHI/Private/D3D12PipelineState.h b/Engine/Source/Runtime/D3D12RHI/Private/D3D12PipelineState.h
--- a/Engine/Source/Runtime/D3D12RHI/Private/D3D12PipelineState.h
+++ b/Engine/Source/Runtime/D3D12RHI/Private/D3D12PipelineState.h
@@ -363,6 +363,22 @@ protected:
 	uint64 ContextSortKey = 0;
 };
 
+struct D3D12ResourceFrameCounter
+{
+	void Set(uint32_t Frame) const
+	{
+		LastUsedFrame.store(Frame, std::memory_order::relaxed);
+	}
+	
+	uint32 Get() const
+	{
+		return LastUsedFrame.load(std::memory_order::relaxed);
+	}
+
+	// Note that the frame counter may be updated from any number of threads, and is not entirely synchronized
+	mutable std::atomic<uint32> LastUsedFrame{0};
+};
+
 struct FD3D12PipelineStateCommonData
 {
 	FD3D12PipelineStateCommonData(const FD3D12RootSignature* InRootSignature, FD3D12PipelineState* InPipelineState);
@@ -370,6 +386,8 @@ struct FD3D12PipelineStateCommonData
 	const FD3D12RootSignature* const RootSignature;
 
 	FD3D12PipelineState* PipelineState;
+
+	D3D12ResourceFrameCounter FrameCounter;
 };
 
 struct FD3D12GraphicsPipelineState : public FRHIGraphicsPipelineState, FD3D12PipelineStateCommonData
diff --git a/Engine/Source/Runtime/D3D12RHI/Private/D3D12RayTracing.cpp b/Engine/Source/Runtime/D3D12RHI/Private/D3D12RayTracing.cpp
--- a/Engine/Source/Runtime/D3D12RHI/Private/D3D12RayTracing.cpp
+++ b/Engine/Source/Runtime/D3D12RHI/Private/D3D12RayTracing.cpp
@@ -391,36 +391,6 @@ inline void UnregisterD3D12RayTracingGeometry(FD3D12RayTracingGeometry* Geometry
 }
 #endif // UE_BUILD_SHIPPING
 
-struct FD3D12ShaderIdentifier
-{
-	uint64 Data[4] = {~0ull, ~0ull, ~0ull, ~0ull};
-
-	// No shader is executed if a shader binding table record with null identifier is encountered.
-	static const FD3D12ShaderIdentifier Null;
-
-	bool operator == (const FD3D12ShaderIdentifier& Other) const
-	{
-		return Data[0] == Other.Data[0]
-			&& Data[1] == Other.Data[1]
-			&& Data[2] == Other.Data[2]
-			&& Data[3] == Other.Data[3];
-	}
-
-	bool operator != (const FD3D12ShaderIdentifier& Other) const
-	{
-		return !(*this == Other);
-	}
-
-	bool IsValid() const
-	{
-		return *this != FD3D12ShaderIdentifier();
-	}
-
-	void SetData(const void* InData)
-	{
-		FMemory::Memcpy(Data, InData, sizeof(Data));
-	}
-};
 
 const FD3D12ShaderIdentifier FD3D12ShaderIdentifier::Null = { 0, 0, 0, 0 };
 
@@ -2597,31 +2567,6 @@ private:
 
 std::atomic_uint64_t FD3D12RayTracingShaderBindingTableInternal::NextUniqueId = 0;
 
-struct FD3D12RayTracingShaderLibrary
-{
-	void Reserve(uint32 NumShaders)
-	{
-		Shaders.Reserve(NumShaders);
-		Identifiers.Reserve(NumShaders);
-	}
-
-	int32 Find(FSHAHash Hash) const
-	{
-		for (int32 Index = 0; Index < Shaders.Num(); ++Index)
-		{
-			if (Hash == Shaders[Index]->GetHash())
-			{
-				return Index;
-			}
-		}
-
-		return INDEX_NONE;
-	}
-
-	TArray<TRefCountPtr<FD3D12RayTracingShader>> Shaders;
-	TArray<FD3D12ShaderIdentifier> Identifiers;
-};
-
 static void CreateSpecializedStateObjects(
 	ID3D12Device5* RayTracingDevice,
 	ID3D12RootSignature* GlobalRootSignature,
@@ -2755,422 +2700,380 @@ static void CreateSpecializedStateObjects(
 	}
 }
 
-class FD3D12RayTracingPipelineState : public FRHIRayTracingPipelineState
+FD3D12RayTracingPipelineState::FD3D12RayTracingPipelineState(FD3D12Device* Device, const FRayTracingPipelineStateInitializer& Initializer) :
+	FRHIRayTracingPipelineState(Initializer),
+	Device(Device)
 {
-public:
+	SCOPE_CYCLE_COUNTER(STAT_RTPSO_CreatePipeline);
+	TRACE_CPUPROFILER_EVENT_SCOPE(RTPSO_CreatePipeline);
 
-	UE_NONCOPYABLE(FD3D12RayTracingPipelineState)
+	checkf(Initializer.GetRayGenTable().Num() > 0 || Initializer.bPartial, TEXT("Ray tracing pipelines must have at leat one ray generation shader."));
+	checkf(Initializer.bPartial || Initializer.GetHitGroupTable().Num() > 0, TEXT("Ray tracing pipelines must have at leat one hit shader."));
 
-	FD3D12RayTracingPipelineState(FD3D12Device* Device, const FRayTracingPipelineStateInitializer& Initializer) : FRHIRayTracingPipelineState(Initializer)
-	{
-		SCOPE_CYCLE_COUNTER(STAT_RTPSO_CreatePipeline);
-		TRACE_CPUPROFILER_EVENT_SCOPE(RTPSO_CreatePipeline);
+	uint64 TotalCreationTime = 0;
+	uint64 CompileTime = 0;
+	uint64 LinkTime = 0;
+	uint32 NumCacheHits = 0;
 
-		checkf(Initializer.GetRayGenTable().Num() > 0 || Initializer.bPartial, TEXT("Ray tracing pipelines must have at leat one ray generation shader."));
-		checkf(Initializer.bPartial || Initializer.GetHitGroupTable().Num() > 0, TEXT("Ray tracing pipelines must have at leat one hit shader."));
+	TotalCreationTime -= FPlatformTime::Cycles64();
 
-		uint64 TotalCreationTime = 0;
-		uint64 CompileTime = 0;
-		uint64 LinkTime = 0;
-		uint32 NumCacheHits = 0;
+	ID3D12Device5* RayTracingDevice = Device->GetDevice5();
 
-		TotalCreationTime -= FPlatformTime::Cycles64();
+	TArrayView<FRHIRayTracingShader*> InitializerHitGroups = Initializer.GetHitGroupTable();
+	TArrayView<FRHIRayTracingShader*> InitializerMissShaders = Initializer.GetMissTable();
+	TArrayView<FRHIRayTracingShader*> InitializerRayGenShaders = Initializer.GetRayGenTable();
+	TArrayView<FRHIRayTracingShader*> InitializerCallableShaders = Initializer.GetCallableTable();
+			
+	FRHIShaderBindingLayout ShaderBindingLayout = Initializer.ShaderBindingLayout ? *Initializer.ShaderBindingLayout : FRHIShaderBindingLayout();
 
-		ID3D12Device5* RayTracingDevice = Device->GetDevice5();
+	const uint32 MaxTotalShaders = InitializerRayGenShaders.Num() + InitializerMissShaders.Num() + InitializerHitGroups.Num() + InitializerCallableShaders.Num();
+	checkf(MaxTotalShaders >= 1, TEXT("Ray tracing pipelines are expected to contain at least one shader"));
 
-		TArrayView<FRHIRayTracingShader*> InitializerHitGroups = Initializer.GetHitGroupTable();
-		TArrayView<FRHIRayTracingShader*> InitializerMissShaders = Initializer.GetMissTable();
-		TArrayView<FRHIRayTracingShader*> InitializerRayGenShaders = Initializer.GetRayGenTable();
-		TArrayView<FRHIRayTracingShader*> InitializerCallableShaders = Initializer.GetCallableTable();
-				
-		FRHIShaderBindingLayout ShaderBindingLayout = Initializer.ShaderBindingLayout ? *Initializer.ShaderBindingLayout : FRHIShaderBindingLayout();
+	FD3D12RayTracingPipelineCache* PipelineCache = Device->GetRayTracingPipelineCache();
 
-		const uint32 MaxTotalShaders = InitializerRayGenShaders.Num() + InitializerMissShaders.Num() + InitializerHitGroups.Num() + InitializerCallableShaders.Num();
-		checkf(MaxTotalShaders >= 1, TEXT("Ray tracing pipelines are expected to contain at least one shader"));
+	// All raygen shaders must share the same global root signature (this is validated below)
 
-		FD3D12RayTracingPipelineCache* PipelineCache = Device->GetRayTracingPipelineCache();
+	GlobalRootSignature = PipelineCache->GetGlobalRootSignature(ShaderBindingLayout);
 
-		// All raygen shaders must share the same global root signature (this is validated below)
+	const FD3D12RayTracingPipelineState* BasePipeline = GRHISupportsRayTracingPSOAdditions 
+		? FD3D12DynamicRHI::ResourceCast(Initializer.BasePipeline.GetReference())
+		: nullptr;
 
-		GlobalRootSignature = PipelineCache->GetGlobalRootSignature(ShaderBindingLayout);
+	if (BasePipeline)
+	{
+		PipelineShaderHashes = BasePipeline->PipelineShaderHashes;
+	}
+	PipelineShaderHashes.Reserve(MaxTotalShaders);
 
-		const FD3D12RayTracingPipelineState* BasePipeline = GRHISupportsRayTracingPSOAdditions 
-			? FD3D12DynamicRHI::ResourceCast(Initializer.BasePipeline.GetReference())
-			: nullptr;
+	TArray<FD3D12RayTracingPipelineCache::FEntry*> UniqueShaderCollections;
+	UniqueShaderCollections.Reserve(MaxTotalShaders);
 
-		if (BasePipeline)
-		{
-			PipelineShaderHashes = BasePipeline->PipelineShaderHashes;
-		}
-		PipelineShaderHashes.Reserve(MaxTotalShaders);
+	FGraphEventArray CompileCompletionList;
+	CompileCompletionList.Reserve(MaxTotalShaders);
 
-		TArray<FD3D12RayTracingPipelineCache::FEntry*> UniqueShaderCollections;
-		UniqueShaderCollections.Reserve(MaxTotalShaders);
+	// Helper function to acquire a D3D12_EXISTING_COLLECTION_DESC for a compiled shader via cache
 
-		FGraphEventArray CompileCompletionList;
-		CompileCompletionList.Reserve(MaxTotalShaders);
+	auto AddShaderCollection = [Device, ShaderBindingLayoutHash = ShaderBindingLayout.GetHash(), GlobalRootSignature = this->GlobalRootSignature, PipelineCache,
+									&UniqueShaderHashes = this->PipelineShaderHashes, &UniqueShaderCollections, &Initializer, &NumCacheHits, &CompileTime,
+									&CompileCompletionList]
+		(FD3D12RayTracingShader* Shader, FD3D12RayTracingPipelineCache::ECollectionType CollectionType)
+	{
+		// verify that that the same shader binding layout is used for all shaders in the RTPSO or not sampling any resources
+		uint32 TotalResourceCount = Shader->ResourceCounts.NumCBs + Shader->ResourceCounts.NumSRVs + Shader->ResourceCounts.NumUAVs + Shader->ResourceCounts.NumSamplers;
+		checkf(TotalResourceCount == 0 || Shader->ShaderBindingLayoutHash == ShaderBindingLayoutHash, TEXT("Raytracing shader with with entry point %s and name %s doesn't match the RTPSO ShaderBindingLayout"), *Shader->EntryPoint, Shader->GetShaderName());
 
-		// Helper function to acquire a D3D12_EXISTING_COLLECTION_DESC for a compiled shader via cache
+		bool bIsAlreadyInSet = false;
+		const uint64 ShaderHash = GetShaderHash64(Shader);
+		UniqueShaderHashes.Add(ShaderHash, &bIsAlreadyInSet);
 
-		auto AddShaderCollection = [Device, ShaderBindingLayoutHash = ShaderBindingLayout.GetHash(), GlobalRootSignature = this->GlobalRootSignature, PipelineCache,
-										&UniqueShaderHashes = this->PipelineShaderHashes, &UniqueShaderCollections, &Initializer, &NumCacheHits, &CompileTime,
-										&CompileCompletionList]
-			(FD3D12RayTracingShader* Shader, FD3D12RayTracingPipelineCache::ECollectionType CollectionType)
-		{
-			// verify that that the same shader binding layout is used for all shaders in the RTPSO or not sampling any resources
-			uint32 TotalResourceCount = Shader->ResourceCounts.NumCBs + Shader->ResourceCounts.NumSRVs + Shader->ResourceCounts.NumUAVs + Shader->ResourceCounts.NumSamplers;
-			checkf(TotalResourceCount == 0 || Shader->ShaderBindingLayoutHash == ShaderBindingLayoutHash, TEXT("Raytracing shader with with entry point %s and name %s doesn't match the RTPSO ShaderBindingLayout"), *Shader->EntryPoint, Shader->GetShaderName());
+		bool bCacheHit = false;
 
-			bool bIsAlreadyInSet = false;
-			const uint64 ShaderHash = GetShaderHash64(Shader);
-			UniqueShaderHashes.Add(ShaderHash, &bIsAlreadyInSet);
+		CompileTime -= FPlatformTime::Cycles64();
 
-			bool bCacheHit = false;
+		FD3D12RayTracingPipelineCache::FEntry* ShaderCacheEntry = PipelineCache->GetOrCompileShader(
+			Device, Shader, GlobalRootSignature,
+			Initializer.MaxAttributeSizeInBytes,
+			Initializer.MaxPayloadSizeInBytes,
+			CollectionType, CompileCompletionList,
+			&bCacheHit);
 
-			CompileTime -= FPlatformTime::Cycles64();
+		CompileTime += FPlatformTime::Cycles64();
 
-			FD3D12RayTracingPipelineCache::FEntry* ShaderCacheEntry = PipelineCache->GetOrCompileShader(
-				Device, Shader, GlobalRootSignature,
-				Initializer.MaxAttributeSizeInBytes,
-				Initializer.MaxPayloadSizeInBytes,
-				CollectionType, CompileCompletionList,
-				&bCacheHit);
+		if (!bIsAlreadyInSet)
+		{
+			UniqueShaderCollections.Add(ShaderCacheEntry);
 
-			CompileTime += FPlatformTime::Cycles64();
+			if (bCacheHit) NumCacheHits++;
+		}
 
-			if (!bIsAlreadyInSet)
-			{
-				UniqueShaderCollections.Add(ShaderCacheEntry);
+		return ShaderCacheEntry;
+	};
 
-				if (bCacheHit) NumCacheHits++;
-			}
 
-			return ShaderCacheEntry;
-		};
+	// If no custom hit groups were provided, then disable SBT indexing and force default shader on all primitives
+	PRAGMA_DISABLE_DEPRECATION_WARNINGS
+	bAllowHitGroupIndexing = Initializer.GetHitGroupTable().Num() ? Initializer.bAllowHitGroupIndexing : false;
+	PRAGMA_ENABLE_DEPRECATION_WARNINGS
 
+	// Add ray generation shaders
 
-		// If no custom hit groups were provided, then disable SBT indexing and force default shader on all primitives
-		PRAGMA_DISABLE_DEPRECATION_WARNINGS
-		bAllowHitGroupIndexing = Initializer.GetHitGroupTable().Num() ? Initializer.bAllowHitGroupIndexing : false;
-		PRAGMA_ENABLE_DEPRECATION_WARNINGS
+	TArray<FD3D12RayTracingPipelineCache::FEntry*> RayGenShaderEntries;
 
-		// Add ray generation shaders
+	RayGenShaders.Reserve(InitializerRayGenShaders.Num());
+	RayGenShaderEntries.Reserve(InitializerRayGenShaders.Num());
+	TMap<FSHAHash, int32> RayGenShaderIndexByHash;
 
-		TArray<FD3D12RayTracingPipelineCache::FEntry*> RayGenShaderEntries;
+	checkf(UniqueShaderCollections.Num() == 0, TEXT("Ray generation shaders are expected to be first in the UniqueShaderCollections list."));
 
-		RayGenShaders.Reserve(InitializerRayGenShaders.Num());
-		RayGenShaderEntries.Reserve(InitializerRayGenShaders.Num());
-		TMap<FSHAHash, int32> RayGenShaderIndexByHash;
+	for (FRHIRayTracingShader* ShaderRHI : InitializerRayGenShaders)
+	{
+		FD3D12RayTracingShader* Shader = FD3D12DynamicRHI::ResourceCast(ShaderRHI);
+		checkf(!Shader->UsesGlobalUniformBuffer(), TEXT("Global uniform buffers are not implemented for ray generation shaders"));
 
-		checkf(UniqueShaderCollections.Num() == 0, TEXT("Ray generation shaders are expected to be first in the UniqueShaderCollections list."));
+		FD3D12RayTracingPipelineCache::FEntry* ShaderCacheEntry = AddShaderCollection(Shader, FD3D12RayTracingPipelineCache::ECollectionType::RayGen);
 
-		for (FRHIRayTracingShader* ShaderRHI : InitializerRayGenShaders)
-		{
-			FD3D12RayTracingShader* Shader = FD3D12DynamicRHI::ResourceCast(ShaderRHI);
-			checkf(!Shader->UsesGlobalUniformBuffer(), TEXT("Global uniform buffers are not implemented for ray generation shaders"));
+		RayGenShaderEntries.Add(ShaderCacheEntry);
+		RayGenShaderIndexByHash.Add(Shader->GetHash(), RayGenShaders.Shaders.Num());
+		RayGenShaders.Shaders.Add(Shader);
+	}
 
-			FD3D12RayTracingPipelineCache::FEntry* ShaderCacheEntry = AddShaderCollection(Shader, FD3D12RayTracingPipelineCache::ECollectionType::RayGen);
+	MaxHitGroupViewDescriptors = 0;
+	MaxLocalRootSignatureSize = 0;
 
-			RayGenShaderEntries.Add(ShaderCacheEntry);
-			RayGenShaderIndexByHash.Add(Shader->GetHash(), RayGenShaders.Shaders.Num());
-			RayGenShaders.Shaders.Add(Shader);
-		}
 
-		MaxHitGroupViewDescriptors = 0;
-		MaxLocalRootSignatureSize = 0;
+	// Add miss shaders
 
+	TArray<FD3D12RayTracingPipelineCache::FEntry*> MissShaderEntries;
+	MissShaders.Reserve(InitializerMissShaders.Num());
+	MissShaderEntries.Reserve(InitializerMissShaders.Num());
 
-		// Add miss shaders
+	for (FRHIRayTracingShader* ShaderRHI : InitializerMissShaders)
+	{
+		FD3D12RayTracingShader* Shader = FD3D12DynamicRHI::ResourceCast(ShaderRHI);
 
-		TArray<FD3D12RayTracingPipelineCache::FEntry*> MissShaderEntries;
-		MissShaders.Reserve(InitializerMissShaders.Num());
-		MissShaderEntries.Reserve(InitializerMissShaders.Num());
+		checkf(Shader, TEXT("A valid ray tracing shader must be provided for all elements in the FRayTracingPipelineStateInitializer miss shader table."));
+		checkf(!Shader->UsesGlobalUniformBuffer(), TEXT("Global uniform buffers are not implemented for ray tracing miss shaders"));
 
-		for (FRHIRayTracingShader* ShaderRHI : InitializerMissShaders)
-		{
-			FD3D12RayTracingShader* Shader = FD3D12DynamicRHI::ResourceCast(ShaderRHI);
+		const uint32 ShaderViewDescriptors = Shader->ResourceCounts.NumSRVs + Shader->ResourceCounts.NumUAVs;
+		MaxHitGroupViewDescriptors = FMath::Max(MaxHitGroupViewDescriptors, ShaderViewDescriptors);
+		MaxLocalRootSignatureSize = FMath::Max(MaxLocalRootSignatureSize, Shader->LocalRootSignature->GetTotalRootSignatureSizeInBytes());
 
-			checkf(Shader, TEXT("A valid ray tracing shader must be provided for all elements in the FRayTracingPipelineStateInitializer miss shader table."));
-			checkf(!Shader->UsesGlobalUniformBuffer(), TEXT("Global uniform buffers are not implemented for ray tracing miss shaders"));
+		FD3D12RayTracingPipelineCache::FEntry* ShaderCacheEntry = AddShaderCollection(Shader, FD3D12RayTracingPipelineCache::ECollectionType::Miss);
 
-			const uint32 ShaderViewDescriptors = Shader->ResourceCounts.NumSRVs + Shader->ResourceCounts.NumUAVs;
-			MaxHitGroupViewDescriptors = FMath::Max(MaxHitGroupViewDescriptors, ShaderViewDescriptors);
-			MaxLocalRootSignatureSize = FMath::Max(MaxLocalRootSignatureSize, Shader->LocalRootSignature->GetTotalRootSignatureSizeInBytes());
+		MissShaderEntries.Add(ShaderCacheEntry);
+		MissShaders.Shaders.Add(Shader);
+	}
 
-			FD3D12RayTracingPipelineCache::FEntry* ShaderCacheEntry = AddShaderCollection(Shader, FD3D12RayTracingPipelineCache::ECollectionType::Miss);
+	// Add hit groups
 
-			MissShaderEntries.Add(ShaderCacheEntry);
-			MissShaders.Shaders.Add(Shader);
-		}
+	TArray<FD3D12RayTracingPipelineCache::FEntry*> HitGroupEntries;
+	HitGroupShaders.Reserve(InitializerHitGroups.Num());
+	HitGroupEntries.Reserve(InitializerHitGroups.Num());
 
-		// Add hit groups
+	for (FRHIRayTracingShader* ShaderRHI : InitializerHitGroups)
+	{
+		FD3D12RayTracingShader* Shader = FD3D12DynamicRHI::ResourceCast(ShaderRHI);
 
-		TArray<FD3D12RayTracingPipelineCache::FEntry*> HitGroupEntries;
-		HitGroupShaders.Reserve(InitializerHitGroups.Num());
-		HitGroupEntries.Reserve(InitializerHitGroups.Num());
+		checkf(Shader, TEXT("A valid ray tracing hit group shader must be provided for all elements in the FRayTracingPipelineStateInitializer hit group table."));
 
-		for (FRHIRayTracingShader* ShaderRHI : InitializerHitGroups)
-		{
-			FD3D12RayTracingShader* Shader = FD3D12DynamicRHI::ResourceCast(ShaderRHI);
+		const uint32 ShaderViewDescriptors = Shader->ResourceCounts.NumSRVs + Shader->ResourceCounts.NumUAVs;
+		MaxHitGroupViewDescriptors = FMath::Max(MaxHitGroupViewDescriptors, ShaderViewDescriptors);
+		MaxLocalRootSignatureSize = FMath::Max(MaxLocalRootSignatureSize, Shader->LocalRootSignature->GetTotalRootSignatureSizeInBytes());
 
-			checkf(Shader, TEXT("A valid ray tracing hit group shader must be provided for all elements in the FRayTracingPipelineStateInitializer hit group table."));
+		FD3D12RayTracingPipelineCache::FEntry* ShaderCacheEntry = AddShaderCollection(Shader, FD3D12RayTracingPipelineCache::ECollectionType::HitGroup);
 
-			const uint32 ShaderViewDescriptors = Shader->ResourceCounts.NumSRVs + Shader->ResourceCounts.NumUAVs;
-			MaxHitGroupViewDescriptors = FMath::Max(MaxHitGroupViewDescriptors, ShaderViewDescriptors);
-			MaxLocalRootSignatureSize = FMath::Max(MaxLocalRootSignatureSize, Shader->LocalRootSignature->GetTotalRootSignatureSizeInBytes());
+		HitGroupEntries.Add(ShaderCacheEntry);
+		HitGroupShaders.Shaders.Add(Shader);
+	}
 
-			FD3D12RayTracingPipelineCache::FEntry* ShaderCacheEntry = AddShaderCollection(Shader, FD3D12RayTracingPipelineCache::ECollectionType::HitGroup);
+	// Add callable shaders
 
-			HitGroupEntries.Add(ShaderCacheEntry);
-			HitGroupShaders.Shaders.Add(Shader);
-		}
+	TArray<FD3D12RayTracingPipelineCache::FEntry*> CallableShaderEntries;
+	CallableShaders.Reserve(InitializerCallableShaders.Num());
+	CallableShaderEntries.Reserve(InitializerCallableShaders.Num());
 
-		// Add callable shaders
+	for (FRHIRayTracingShader* ShaderRHI : InitializerCallableShaders)
+	{
+		FD3D12RayTracingShader* Shader = FD3D12DynamicRHI::ResourceCast(ShaderRHI);
 
-		TArray<FD3D12RayTracingPipelineCache::FEntry*> CallableShaderEntries;
-		CallableShaders.Reserve(InitializerCallableShaders.Num());
-		CallableShaderEntries.Reserve(InitializerCallableShaders.Num());
+		checkf(Shader, TEXT("A valid ray tracing shader must be provided for all elements in the FRayTracingPipelineStateInitializer callable shader table."));
+		checkf(!Shader->UsesGlobalUniformBuffer(), TEXT("Global uniform buffers are not implemented for ray tracing callable shaders"));
 
-		for (FRHIRayTracingShader* ShaderRHI : InitializerCallableShaders)
-		{
-			FD3D12RayTracingShader* Shader = FD3D12DynamicRHI::ResourceCast(ShaderRHI);
+		const uint32 ShaderViewDescriptors = Shader->ResourceCounts.NumSRVs + Shader->ResourceCounts.NumUAVs;
+		MaxHitGroupViewDescriptors = FMath::Max(MaxHitGroupViewDescriptors, ShaderViewDescriptors);
+		MaxLocalRootSignatureSize = FMath::Max(MaxLocalRootSignatureSize, Shader->LocalRootSignature->GetTotalRootSignatureSizeInBytes());
 
-			checkf(Shader, TEXT("A valid ray tracing shader must be provided for all elements in the FRayTracingPipelineStateInitializer callable shader table."));
-			checkf(!Shader->UsesGlobalUniformBuffer(), TEXT("Global uniform buffers are not implemented for ray tracing callable shaders"));
+		FD3D12RayTracingPipelineCache::FEntry* ShaderCacheEntry = AddShaderCollection(Shader, FD3D12RayTracingPipelineCache::ECollectionType::Callable);
 
-			const uint32 ShaderViewDescriptors = Shader->ResourceCounts.NumSRVs + Shader->ResourceCounts.NumUAVs;
-			MaxHitGroupViewDescriptors = FMath::Max(MaxHitGroupViewDescriptors, ShaderViewDescriptors);
-			MaxLocalRootSignatureSize = FMath::Max(MaxLocalRootSignatureSize, Shader->LocalRootSignature->GetTotalRootSignatureSizeInBytes());
+		CallableShaderEntries.Add(ShaderCacheEntry);
+		CallableShaders.Shaders.Add(Shader);
+	}
 
-			FD3D12RayTracingPipelineCache::FEntry* ShaderCacheEntry = AddShaderCollection(Shader, FD3D12RayTracingPipelineCache::ECollectionType::Callable);
+	check(Initializer.GetMaxLocalBindingDataSize() >= MaxLocalRootSignatureSize);
 
-			CallableShaderEntries.Add(ShaderCacheEntry);
-			CallableShaders.Shaders.Add(Shader);
-		}
+	// Wait for all compilation tasks to be complete and then gather the compiled collection descriptors
 
-		check(Initializer.GetMaxLocalBindingDataSize() >= MaxLocalRootSignatureSize);
+	CompileTime -= FPlatformTime::Cycles64();
 
-		// Wait for all compilation tasks to be complete and then gather the compiled collection descriptors
+	FTaskGraphInterface::Get().WaitUntilTasksComplete(CompileCompletionList);
 
-		CompileTime -= FPlatformTime::Cycles64();
+	CompileTime += FPlatformTime::Cycles64();
+	
+	if (Initializer.bPartial)
+	{
+		// Partial pipelines don't have a linking phase, so exit immediately after compilation tasks are complete.
+		return;
+	}
 
-		FTaskGraphInterface::Get().WaitUntilTasksComplete(CompileCompletionList);
+	TArray<D3D12_EXISTING_COLLECTION_DESC> UniqueShaderCollectionDescs;
+	UniqueShaderCollectionDescs.Reserve(MaxTotalShaders);
+	for (FD3D12RayTracingPipelineCache::FEntry* Entry : UniqueShaderCollections)
+	{
+		UniqueShaderCollectionDescs.Add(Entry->GetCollectionDesc());
+	}
 
-		CompileTime += FPlatformTime::Cycles64();
+	// Link final RTPSO from shader collections
 
-		if (Initializer.bPartial)
-		{
-			// Partial pipelines don't have a linking phase, so exit immediately after compilation tasks are complete.
-			return;
-		}
+	LinkTime -= FPlatformTime::Cycles64();
 
-		TArray<D3D12_EXISTING_COLLECTION_DESC> UniqueShaderCollectionDescs;
-		UniqueShaderCollectionDescs.Reserve(MaxTotalShaders);
-		for (FD3D12RayTracingPipelineCache::FEntry* Entry : UniqueShaderCollections)
+	// Extending RTPSOs is currently not compatible with PSO specializations
+	if (BasePipeline && GRayTracingSpecializeStateObjects == 0)
+	{
+		if (UniqueShaderCollectionDescs.Num() == 0)
 		{
-			UniqueShaderCollectionDescs.Add(Entry->GetCollectionDesc());
+			// New PSO does not actually have any new shaders that were not in the base
+			StateObject = BasePipeline->StateObject.GetReference();
 		}
-
-		// Link final RTPSO from shader collections
-
-		LinkTime -= FPlatformTime::Cycles64();
-
-		// Extending RTPSOs is currently not compatible with PSO specializations
-		if (BasePipeline && GRayTracingSpecializeStateObjects == 0)
+		else
 		{
-			if (UniqueShaderCollectionDescs.Num() == 0)
-			{
-				// New PSO does not actually have any new shaders that were not in the base
-				StateObject = BasePipeline->StateObject.GetReference();
-			}
-			else
-			{
-
-				TArray<D3D12_STATE_SUBOBJECT> Subobjects;
 
-				int32 SubobjectIndex = 0;
-				Subobjects.Reserve(UniqueShaderCollectionDescs.Num() + 1);
+			TArray<D3D12_STATE_SUBOBJECT> Subobjects;
 
-				D3D12_STATE_OBJECT_CONFIG StateObjectConfig = {};
-				StateObjectConfig.Flags = D3D12_STATE_OBJECT_FLAG_ALLOW_STATE_OBJECT_ADDITIONS;
-				Subobjects.Add(D3D12_STATE_SUBOBJECT{ D3D12_STATE_SUBOBJECT_TYPE_STATE_OBJECT_CONFIG, &StateObjectConfig });
+			int32 SubobjectIndex = 0;
+			Subobjects.Reserve(UniqueShaderCollectionDescs.Num() + 1);
 
-				for (const D3D12_EXISTING_COLLECTION_DESC& Collection : UniqueShaderCollectionDescs)
-				{
-					Subobjects.Add(D3D12_STATE_SUBOBJECT{ D3D12_STATE_SUBOBJECT_TYPE_EXISTING_COLLECTION, &Collection });
-				}
+			D3D12_STATE_OBJECT_CONFIG StateObjectConfig = {};
+			StateObjectConfig.Flags = D3D12_STATE_OBJECT_FLAG_ALLOW_STATE_OBJECT_ADDITIONS;
+			Subobjects.Add(D3D12_STATE_SUBOBJECT{ D3D12_STATE_SUBOBJECT_TYPE_STATE_OBJECT_CONFIG, &StateObjectConfig });
 
-				D3D12_STATE_OBJECT_DESC Desc = {};
-				Desc.Type = D3D12_STATE_OBJECT_TYPE_RAYTRACING_PIPELINE;
-				Desc.NumSubobjects = Subobjects.Num();
-				Desc.pSubobjects = Subobjects.GetData();
+			for (const D3D12_EXISTING_COLLECTION_DESC& Collection : UniqueShaderCollectionDescs)
+			{
+				Subobjects.Add(D3D12_STATE_SUBOBJECT{ D3D12_STATE_SUBOBJECT_TYPE_EXISTING_COLLECTION, &Collection });
+			}
 
-				ID3D12Device7* Device7 = Device->GetDevice7();
+			D3D12_STATE_OBJECT_DESC Desc = {};
+			Desc.Type = D3D12_STATE_OBJECT_TYPE_RAYTRACING_PIPELINE;
+			Desc.NumSubobjects = Subobjects.Num();
+			Desc.pSubobjects = Subobjects.GetData();
 
-				VERIFYD3D12RESULT(Device7->AddToStateObject(&Desc,
-					BasePipeline->StateObject.GetReference(),
-					IID_PPV_ARGS(StateObject.GetInitReference())));
-			}
-		}
-		else
-		{
-			StateObject = CreateRayTracingStateObject(
-				RayTracingDevice,
-				{}, // Libraries,
-				{}, // LibraryExports,
-				Initializer.MaxAttributeSizeInBytes,
-				Initializer.MaxPayloadSizeInBytes,
-				{}, // HitGroups
-				GlobalRootSignature,
-				{}, // LocalRootSignatures
-				{}, // LocalRootSignatureAssociations,
-				UniqueShaderCollectionDescs,
-				D3D12_STATE_OBJECT_TYPE_RAYTRACING_PIPELINE);
-		}
+			ID3D12Device7* Device7 = Device->GetDevice7();
 
-		if (GRayTracingSpecializeStateObjects != 0 && Initializer.GetRayGenTable().Num() > 1)
-		{
-			CreateSpecializedStateObjects(
-				RayTracingDevice,
-				GlobalRootSignature,
-				Initializer.MaxAttributeSizeInBytes,
-				Initializer.MaxPayloadSizeInBytes,
-				RayGenShaders,
-				UniqueShaderCollections,
-				RayGenShaderIndexByHash,
-				SpecializedStateObjects, // out param
-				SpecializationIndices // out param
-			);
+			VERIFYD3D12RESULT(Device7->AddToStateObject(&Desc,
+				BasePipeline->StateObject.GetReference(),
+				IID_PPV_ARGS(StateObject.GetInitReference())));
 		}
+	}
+	else
+	{
+		StateObject = CreateRayTracingStateObject(
+			RayTracingDevice,
+			{}, // Libraries,
+			{}, // LibraryExports,
+			Initializer.MaxAttributeSizeInBytes,
+			Initializer.MaxPayloadSizeInBytes,
+			{}, // HitGroups
+			GlobalRootSignature,
+			{}, // LocalRootSignatures
+			{}, // LocalRootSignatureAssociations,
+			UniqueShaderCollectionDescs,
+			D3D12_STATE_OBJECT_TYPE_RAYTRACING_PIPELINE);
+	}
 
-		LinkTime += FPlatformTime::Cycles64();
+	if (GRayTracingSpecializeStateObjects != 0 && Initializer.GetRayGenTable().Num() > 1)
+	{
+		CreateSpecializedStateObjects(
+			RayTracingDevice,
+			GlobalRootSignature,
+			Initializer.MaxAttributeSizeInBytes,
+			Initializer.MaxPayloadSizeInBytes,
+			RayGenShaders,
+			UniqueShaderCollections,
+			RayGenShaderIndexByHash,
+			SpecializedStateObjects, // out param
+			SpecializationIndices // out param
+		);
+	}
 
-		HRESULT QueryInterfaceResult = StateObject->QueryInterface(IID_PPV_ARGS(PipelineProperties.GetInitReference()));
-		checkf(SUCCEEDED(QueryInterfaceResult), TEXT("Failed to query pipeline properties from the ray tracing pipeline state object. Result=%08x"), QueryInterfaceResult);
+	LinkTime += FPlatformTime::Cycles64();
 
-		// Query shader identifiers from the pipeline state object
+	HRESULT QueryInterfaceResult = StateObject->QueryInterface(IID_PPV_ARGS(PipelineProperties.GetInitReference()));
+	checkf(SUCCEEDED(QueryInterfaceResult), TEXT("Failed to query pipeline properties from the ray tracing pipeline state object. Result=%08x"), QueryInterfaceResult);
 
-		check(HitGroupEntries.Num() == InitializerHitGroups.Num());
+	// Query shader identifiers from the pipeline state object
 
-		auto GetEntryShaderIdentifier = [Properties = PipelineProperties.GetReference()](FD3D12RayTracingPipelineCache::FEntry* Entry) -> FD3D12ShaderIdentifier
-		{
-			if (Entry->Identifier.IsValid())
-			{
-				return Entry->Identifier;
-			}
-			else
-			{
-				return GetShaderIdentifier(Properties, Entry->GetPrimaryExportNameChars());
-			}
-		};
+	check(HitGroupEntries.Num() == InitializerHitGroups.Num());
 
-		HitGroupShaders.Identifiers.SetNumUninitialized(InitializerHitGroups.Num());
-		for (int32 HitGroupIndex = 0; HitGroupIndex < HitGroupEntries.Num(); ++HitGroupIndex)
-		{
-			HitGroupShaders.Identifiers[HitGroupIndex] = GetEntryShaderIdentifier(HitGroupEntries[HitGroupIndex]);
-		}
-
-		RayGenShaders.Identifiers.SetNumUninitialized(RayGenShaderEntries.Num());
-		for (int32 ShaderIndex = 0; ShaderIndex < RayGenShaderEntries.Num(); ++ShaderIndex)
+	auto GetEntryShaderIdentifier = [Properties = PipelineProperties.GetReference()](FD3D12RayTracingPipelineCache::FEntry* Entry) -> FD3D12ShaderIdentifier
+	{
+		if (Entry->Identifier.IsValid())
 		{
-			RayGenShaders.Identifiers[ShaderIndex] = GetEntryShaderIdentifier(RayGenShaderEntries[ShaderIndex]);
+			return Entry->Identifier;
 		}
-
-		MissShaders.Identifiers.SetNumUninitialized(MissShaderEntries.Num());
-		for (int32 ShaderIndex = 0; ShaderIndex < MissShaderEntries.Num(); ++ShaderIndex)
+		else
 		{
-			MissShaders.Identifiers[ShaderIndex] = GetEntryShaderIdentifier(MissShaderEntries[ShaderIndex]);
+			return GetShaderIdentifier(Properties, Entry->GetPrimaryExportNameChars());
 		}
+	};
 
-		CallableShaders.Identifiers.SetNumUninitialized(CallableShaderEntries.Num());
-		for (int32 ShaderIndex = 0; ShaderIndex < CallableShaderEntries.Num(); ++ShaderIndex)
-		{
-			CallableShaders.Identifiers[ShaderIndex] = GetEntryShaderIdentifier(CallableShaderEntries[ShaderIndex]);
-		}
+	HitGroupShaders.Identifiers.SetNumUninitialized(InitializerHitGroups.Num());
+	for (int32 HitGroupIndex = 0; HitGroupIndex < HitGroupEntries.Num(); ++HitGroupIndex)
+	{
+		HitGroupShaders.Identifiers[HitGroupIndex] = GetEntryShaderIdentifier(HitGroupEntries[HitGroupIndex]);
+	}
 
-		PipelineStackSize = PipelineProperties->GetPipelineStackSize();
+	RayGenShaders.Identifiers.SetNumUninitialized(RayGenShaderEntries.Num());
+	for (int32 ShaderIndex = 0; ShaderIndex < RayGenShaderEntries.Num(); ++ShaderIndex)
+	{
+		RayGenShaders.Identifiers[ShaderIndex] = GetEntryShaderIdentifier(RayGenShaderEntries[ShaderIndex]);
+	}
 
-		TotalCreationTime += FPlatformTime::Cycles64();
+	MissShaders.Identifiers.SetNumUninitialized(MissShaderEntries.Num());
+	for (int32 ShaderIndex = 0; ShaderIndex < MissShaderEntries.Num(); ++ShaderIndex)
+	{
+		MissShaders.Identifiers[ShaderIndex] = GetEntryShaderIdentifier(MissShaderEntries[ShaderIndex]);
+	}
 
-		// Report stats for pipelines that take a long time to create
+	CallableShaders.Identifiers.SetNumUninitialized(CallableShaderEntries.Num());
+	for (int32 ShaderIndex = 0; ShaderIndex < CallableShaderEntries.Num(); ++ShaderIndex)
+	{
+		CallableShaders.Identifiers[ShaderIndex] = GetEntryShaderIdentifier(CallableShaderEntries[ShaderIndex]);
+	}
 
-#if !NO_LOGGING
+	PipelineStackSize = PipelineProperties->GetPipelineStackSize();
 
-		// Gather PSO stats
-		ShaderStats.Reserve(UniqueShaderCollections.Num());
-		for (FD3D12RayTracingPipelineCache::FEntry* Entry : UniqueShaderCollections)
-		{
-			FShaderStats Stats;
-			Stats.Name = *(Entry->Shader->EntryPoint);
-			Stats.ShaderSize = Entry->Shader->Code.Num();
-			Stats.CompileTimeMS = Entry->CompileTimeMS;
+	TotalCreationTime += FPlatformTime::Cycles64();
 
-		#if PLATFORM_WINDOWS
-			if (Entry->Shader->GetFrequency() == SF_RayGen)
-			{
-				Stats.StackSize = uint32(PipelineProperties->GetShaderStackSize(*(Entry->ExportNames[0])));
-			}
-		#endif // PLATFORM_WINDOWS
+	// Report stats for pipelines that take a long time to create
 
-			ShaderStats.Add(Stats);
-		}
+#if !NO_LOGGING
 
-		ShaderStats.Sort([](const FShaderStats& A, const FShaderStats& B) { return B.CompileTimeMS < A.CompileTimeMS; });
+	// Gather PSO stats
+	ShaderStats.Reserve(UniqueShaderCollections.Num());
+	for (FD3D12RayTracingPipelineCache::FEntry* Entry : UniqueShaderCollections)
+	{
+		FShaderStats Stats;
+		Stats.Name = *(Entry->Shader->EntryPoint);
+		Stats.ShaderSize = Entry->Shader->Code.Num();
+		Stats.CompileTimeMS = Entry->CompileTimeMS;
 
-		const double TotalCreationTimeMS = 1000.0 * FPlatformTime::ToSeconds64(TotalCreationTime);
-		// log a message if the RTPSO creating took a long time (with a higher threshold if this was a background compilation)
-		const double CreationTimeWarningThresholdMS = 10.0;
-		const double BackgroundCreatingTimeWarningThresholdMS = 1000.0;
-		if (TotalCreationTimeMS > (Initializer.bBackgroundCompilation ? CreationTimeWarningThresholdMS : BackgroundCreatingTimeWarningThresholdMS))
+	#if PLATFORM_WINDOWS
+		if (Entry->Shader->GetFrequency() == SF_RayGen)
 		{
-			const double CompileTimeMS = 1000.0 * FPlatformTime::ToSeconds64(CompileTime);
-			const double LinkTimeMS = 1000.0 * FPlatformTime::ToSeconds64(LinkTime);
-			const uint32 NumUniqueShaders = UniqueShaderCollections.Num();
-			UE_LOG(LogD3D12RHI, Log,
-				TEXT("Creating RTPSO with %d shaders (%d cached, %d new) took %.2f ms. Compile time %.2f ms, link time %.2f ms."),
-				PipelineShaderHashes.Num(), NumCacheHits, NumUniqueShaders - NumCacheHits, (float)TotalCreationTimeMS, (float)CompileTimeMS, (float)LinkTimeMS);
+			Stats.StackSize = uint32(PipelineProperties->GetShaderStackSize(*(Entry->ExportNames[0])));
 		}
+	#endif // PLATFORM_WINDOWS
 
-#endif //!NO_LOGGING
+		ShaderStats.Add(Stats);
 	}
 
-	FD3D12RayTracingShaderLibrary RayGenShaders;
-	FD3D12RayTracingShaderLibrary MissShaders;
-	FD3D12RayTracingShaderLibrary HitGroupShaders;
-	FD3D12RayTracingShaderLibrary CallableShaders;
-
-	ID3D12RootSignature* GlobalRootSignature = nullptr;
-
-	TRefCountPtr<ID3D12StateObject> StateObject;
-	TRefCountPtr<ID3D12StateObjectProperties> PipelineProperties;
-
-	// Maps raygen shader index to a specialized state object (may be -1 if no specialization is used for a shader)
-	TArray<int32> SpecializationIndices;
-
-	// State objects with raygen shaders grouped by occupancy
-	TArray<TRefCountPtr<ID3D12StateObject>> SpecializedStateObjects;
-
-	UE_DEPRECATED(5.5, "bAllowHitGroupIndexing is now stored in the ShaderBindingTable.")
-	bool bAllowHitGroupIndexing = true;
-
-	uint32 MaxLocalRootSignatureSize = 0;
-	uint32 MaxHitGroupViewDescriptors = 0;
-
-	TSet<uint64> PipelineShaderHashes;
+	ShaderStats.Sort([](const FShaderStats& A, const FShaderStats& B) { return B.CompileTimeMS < A.CompileTimeMS; });
 
-	uint32 PipelineStackSize = 0;
-
-#if !NO_LOGGING
-	struct FShaderStats
+	const double TotalCreationTimeMS = 1000.0 * FPlatformTime::ToSeconds64(TotalCreationTime);
+	// log a message if the RTPSO creating took a long time (with a higher threshold if this was a background compilation)
+	const double CreationTimeWarningThresholdMS = 10.0;
+	const double BackgroundCreatingTimeWarningThresholdMS = 1000.0;
+	if (TotalCreationTimeMS > (Initializer.bBackgroundCompilation ? CreationTimeWarningThresholdMS : BackgroundCreatingTimeWarningThresholdMS))
 	{
-		const TCHAR* Name = nullptr;
-		float CompileTimeMS = 0;
-		uint32 StackSize = 0;
-		uint32 ShaderSize = 0;
-	};
-	TArray<FShaderStats> ShaderStats;
-#endif // !NO_LOGGING
-};
+		const double CompileTimeMS = 1000.0 * FPlatformTime::ToSeconds64(CompileTime);
+		const double LinkTimeMS = 1000.0 * FPlatformTime::ToSeconds64(LinkTime);
+		const uint32 NumUniqueShaders = UniqueShaderCollections.Num();
+		UE_LOG(LogD3D12RHI, Log,
+			TEXT("Creating RTPSO with %d shaders (%d cached, %d new) took %.2f ms. Compile time %.2f ms, link time %.2f ms."),
+			PipelineShaderHashes.Num(), NumCacheHits, NumUniqueShaders - NumCacheHits, (float)TotalCreationTimeMS, (float)CompileTimeMS, (float)LinkTimeMS);
+	}
+
+#endif //!NO_LOGGING
+}
 
 class FD3D12RayTracingShaderBindingTable : public FRHIShaderBindingTable, public FD3D12AdapterChild
 {
@@ -5809,7 +5712,9 @@ static void DispatchRays(FD3D12CommandContext& CommandContext,
 		{
 			RayTracingStateObject = Pipeline->StateObject.GetReference();
 		}
-;
+
+		Pipeline->FrameCounter.Set(CommandContext.GetFrameFenceCounter());
+		
 		CommandContext.RayTracingCommandList()->SetPipelineState1(RayTracingStateObject);
 
 		if (DispatchRaysDescBuffer)
diff --git a/Engine/Source/Runtime/D3D12RHI/Private/D3D12RayTracing.h b/Engine/Source/Runtime/D3D12RHI/Private/D3D12RayTracing.h
--- a/Engine/Source/Runtime/D3D12RHI/Private/D3D12RayTracing.h
+++ b/Engine/Source/Runtime/D3D12RHI/Private/D3D12RayTracing.h
@@ -22,6 +22,111 @@ struct ID3D12RayTracingGeometryUpdateListener
 	virtual void HitGroupParametersUpdated(FD3D12RayTracingGeometry* InGeometry) = 0;
 };
 
+struct FD3D12ShaderIdentifier
+{
+	uint64 Data[4] = {~0ull, ~0ull, ~0ull, ~0ull};
+
+	// No shader is executed if a shader binding table record with null identifier is encountered.
+	static const FD3D12ShaderIdentifier Null;
+
+	bool operator == (const FD3D12ShaderIdentifier& Other) const
+	{
+		return Data[0] == Other.Data[0]
+			&& Data[1] == Other.Data[1]
+			&& Data[2] == Other.Data[2]
+			&& Data[3] == Other.Data[3];
+	}
+
+	bool operator != (const FD3D12ShaderIdentifier& Other) const
+	{
+		return !(*this == Other);
+	}
+
+	bool IsValid() const
+	{
+		return *this != FD3D12ShaderIdentifier();
+	}
+
+	void SetData(const void* InData)
+	{
+		FMemory::Memcpy(Data, InData, sizeof(Data));
+	}
+};
+
+struct FD3D12RayTracingShaderLibrary
+{
+	void Reserve(uint32 NumShaders)
+	{
+		Shaders.Reserve(NumShaders);
+		Identifiers.Reserve(NumShaders);
+	}
+
+	int32 Find(FSHAHash Hash) const
+	{
+		for (int32 Index = 0; Index < Shaders.Num(); ++Index)
+		{
+			if (Hash == Shaders[Index]->GetHash())
+			{
+				return Index;
+			}
+		}
+
+		return INDEX_NONE;
+	}
+
+	TArray<TRefCountPtr<FD3D12RayTracingShader>> Shaders;
+	TArray<FD3D12ShaderIdentifier> Identifiers;
+};
+
+class FD3D12RayTracingPipelineState : public FRHIRayTracingPipelineState
+{
+public:
+	UE_NONCOPYABLE(FD3D12RayTracingPipelineState)
+
+	FD3D12RayTracingPipelineState(FD3D12Device* Device, const FRayTracingPipelineStateInitializer& Initializer);
+
+	FD3D12Device* Device;
+	
+	FD3D12RayTracingShaderLibrary RayGenShaders;
+	FD3D12RayTracingShaderLibrary MissShaders;
+	FD3D12RayTracingShaderLibrary HitGroupShaders;
+	FD3D12RayTracingShaderLibrary CallableShaders;
+
+	ID3D12RootSignature* GlobalRootSignature = nullptr;
+
+	TRefCountPtr<ID3D12StateObject> StateObject;
+	TRefCountPtr<ID3D12StateObjectProperties> PipelineProperties;
+
+	// Maps raygen shader index to a specialized state object (may be -1 if no specialization is used for a shader)
+	TArray<int32> SpecializationIndices;
+
+	// State objects with raygen shaders grouped by occupancy
+	TArray<TRefCountPtr<ID3D12StateObject>> SpecializedStateObjects;
+
+	UE_DEPRECATED(5.5, "bAllowHitGroupIndexing is now stored in the ShaderBindingTable.")
+	bool bAllowHitGroupIndexing = true;
+
+	uint32 MaxLocalRootSignatureSize = 0;
+	uint32 MaxHitGroupViewDescriptors = 0;
+
+	TSet<uint64> PipelineShaderHashes;
+
+	uint32 PipelineStackSize = 0;
+
+#if !NO_LOGGING
+	struct FShaderStats
+	{
+		const TCHAR* Name = nullptr;
+		float CompileTimeMS = 0;
+		uint32 StackSize = 0;
+		uint32 ShaderSize = 0;
+	};
+	TArray<FShaderStats> ShaderStats;
+#endif // !NO_LOGGING
+
+	D3D12ResourceFrameCounter FrameCounter;
+};
+
 class FD3D12RayTracingGeometry : public FRHIRayTracingGeometry, public FD3D12AdapterChild, public FD3D12ShaderResourceRenameListener, public FNoncopyable
 {
 public:
diff --git a/Engine/Source/Runtime/D3D12RHI/Private/D3D12WorkGraph.cpp b/Engine/Source/Runtime/D3D12RHI/Private/D3D12WorkGraph.cpp
--- a/Engine/Source/Runtime/D3D12RHI/Private/D3D12WorkGraph.cpp
+++ b/Engine/Source/Runtime/D3D12RHI/Private/D3D12WorkGraph.cpp
@@ -22,6 +22,7 @@ static FAutoConsoleVariableRef CVarShaderBundleSkipDispatch(
 );
 
 FD3D12WorkGraphPipelineState::FD3D12WorkGraphPipelineState(FD3D12Device* Device, const FWorkGraphPipelineStateInitializer& Initializer)
+	: Device(Device)
 {
 #if D3D12_RHI_WORKGRAPHS
 	
@@ -292,6 +293,16 @@ FD3D12WorkGraphPipelineState::FD3D12WorkGraphPipelineState(FD3D12Device* Device,
 	BackingMemoryAddressRange.StartAddress = BackingMemoryBufferResource->GetGPUVirtualAddress();
 	BackingMemoryAddressRange.SizeInBytes = MemoryRequirements.MaxSizeInBytes;
 
+#if NV_AFTERMATH
+	if (UE::RHICore::Nvidia::Aftermath::IsShaderRegistrationEnabled())
+	{
+		// Copy shader table for late association
+		for (FRHIWorkGraphShader* Shader : Initializer.GetShaderTable())
+		{
+			Shaders.Add(TRefCountPtr<FRHIShader>(Shader));
+		}
+	}
+#endif // NV_AFTERMATH
 #endif // D3D12_RHI_WORKGRAPHS
 }
 
@@ -1389,6 +1400,8 @@ void FD3D12CommandContext::DispatchWorkGraphShaderBundle(
 	SetProgramDesc.WorkGraph.NodeLocalRootArgumentsTable = NodeLocalRootArgumentsTable;
 	GraphicsCommandList10()->SetProgram(&SetProgramDesc);
 
+	Pipeline->FrameCounter.Set(GetFrameFenceCounter());
+
 	const uint32 RecordCount = Dispatches.Num();
 	FDispatchShaderBundleWorkGraph::FEntryNodeRecord InputRecord = FDispatchShaderBundleWorkGraph::MakeInputRecord(RecordCount, ShaderBundle->ArgOffset, ShaderBundle->ArgStride, RecordArgBufferBindlessHandle);
 
diff --git a/Engine/Source/Runtime/D3D12RHI/Private/D3D12WorkGraph.h b/Engine/Source/Runtime/D3D12RHI/Private/D3D12WorkGraph.h
--- a/Engine/Source/Runtime/D3D12RHI/Private/D3D12WorkGraph.h
+++ b/Engine/Source/Runtime/D3D12RHI/Private/D3D12WorkGraph.h
@@ -4,6 +4,7 @@
 
 #include "D3D12RHICommon.h"
 #include "RHIResources.h"
+#include "D3D12PipelineState.h"
 
 class FD3D12Device;
 class FWorkGraphPipelineStateInitializer;
@@ -16,6 +17,8 @@ public:
 	FD3D12WorkGraphPipelineState(FD3D12Device* Device, const FWorkGraphPipelineStateInitializer& Initializer);
 
 #if D3D12_RHI_WORKGRAPHS
+	FD3D12Device* Device;
+	
 	TRefCountPtr<ID3D12RootSignature> RootSignature;
 	TRefCountPtr<ID3D12StateObject>	StateObject;
 	
@@ -29,5 +32,11 @@ public:
 	TMap<FString, uint32> NodeCountPerName;
 
 	bool bInitialized = false;
+
+	D3D12ResourceFrameCounter FrameCounter;
+	
+#if NV_AFTERMATH
+	TArray<TRefCountPtr<FRHIShader>> Shaders;
+#endif // NV_AFTERMATH
 #endif
 };
diff --git a/Engine/Source/Runtime/RHI/Private/PipelineStateCache.cpp b/Engine/Source/Runtime/RHI/Private/PipelineStateCache.cpp
--- a/Engine/Source/Runtime/RHI/Private/PipelineStateCache.cpp
+++ b/Engine/Source/Runtime/RHI/Private/PipelineStateCache.cpp
@@ -1374,6 +1374,8 @@ public:
 		FScopeVerifyIncrement S(VerifyMutex);
 #endif
 
+		FRWScopeLock InterruptGuard(InterruptLock, SLT_ReadOnly);
+
 		// Select the maps to use.
 		FPipelineStateCacheType* LocalCurrentMap = CurrentMap;
 		FPipelineStateCacheType* LocalBackfillMap = BackfillMap;
@@ -1423,6 +1425,9 @@ public:
 #if PIPELINESTATECACHE_VERIFYTHREADSAFE
 		FScopeVerifyIncrement S(VerifyMutex);
 #endif
+
+		FRWScopeLock InterruptGuard(InterruptLock, SLT_ReadOnly);
+		
 		// everything is added to the local cache then at end of frame we consolidate them all
 		TMap<TMyKey, TMyValue>& LocalCache = GetLocalCache();
 
@@ -1432,11 +1437,48 @@ public:
 		return true;
 	}
 
+	void GetResources(TArray<TRefCountPtr<FRHIResource>>& OutResources, bool bConsolidateWithInterrupt = true)
+	{
+		FRWScopeLock InterruptGuard(InterruptLock, SLT_Write);
+		
+		if (bConsolidateWithInterrupt)
+		{
+			// Wait for any in-flight consolidation
+			WaitAndFinishAsyncCacheConsolidation();
+
+			bIsInterrupt = true;
+			
+			// Kick off a new one
+			FlushResources(false);
+			WaitAndFinishAsyncCacheConsolidation();
+
+			bIsInterrupt = false;
+		}
+		
+		for (auto&& [Desc, State] : *CurrentMap)
+		{
+			OutResources.Add(TRefCountPtr<FRHIResource>(State->RHIPipeline));
+		} 
+	}
+
 	// Call from the Render Thread
 	void FlushResources(bool bInDiscardAndSwap)
 	{
 		SCOPED_NAMED_EVENT(ConsolidateThreadedCaches, FColor::Turquoise);
 
+		if (!bIsInterrupt)
+		{
+			InterruptLock.ReadLock();
+		}
+		
+		ON_SCOPE_EXIT
+		{
+			if (!bIsInterrupt)
+			{
+				InterruptLock.ReadUnlock();
+			}
+		};
+
 		bPendingDiscardAndSwap |= bInDiscardAndSwap;
 
 #if PIPELINESTATECACHE_VERIFYTHREADSAFE
@@ -1545,6 +1587,18 @@ public:
 	}
 
 private:
+	void WaitAndFinishAsyncCacheConsolidation()
+	{
+		if (!RHICompletionEvent)
+		{
+			return;
+		}
+		
+		RHICompletionEvent->Wait();
+		RHICompletionEvent = nullptr;
+		
+		FinishAsyncCacheConsolidation();
+	}
 
 	void WaitTasksComplete(FPipelineStateCacheType* PipelineStateCache)
 	{
@@ -1640,8 +1694,13 @@ private:
 		RHICompletionEvent = FGraphEvent::CreateGraphEvent();
 		RHICompletionEvent->SetDebugName(TEXT("AsyncCacheConsolidation"));
 
-		// Add the completion of the RHI cache consolidation as a prerequisite for the next RHI dispatch.
-		GRHICommandList.AddNextDispatchPrerequisite(RHICompletionEvent);
+		FGraphEventArray Prerequisites;
+		if (!bIsInterrupt)
+		{
+			// Add the completion of the RHI cache consolidation as a prerequisite for the next RHI dispatch.
+			GRHICommandList.AddNextDispatchPrerequisite(RHICompletionEvent);
+			Prerequisites.Add(GRHICommandList.GetCompletionEvent());
+		}
 
 		// Flush Render Thread local caches and consolidate them into a single map.
 		ConsolidatePipelineStates(NewRenderThreadPipelineStates, RenderThreadsPipelineStateCache);
@@ -1658,7 +1717,7 @@ private:
 				RHIAsyncCacheConsolidation();
 			},
 			TStatId(),
-			GRHICommandList.GetCompletionEvent(),
+			&Prerequisites,
 			ENamedThreads::AnyHiPriThreadHiPriTask);
 	}
 
@@ -1928,6 +1987,13 @@ private:
 	TArray<FPipelineStateCacheType*> RenderThreadsPipelineStateCache;
 	TArray<FPipelineStateCacheType*> RHIThreadsPipelineStateCache;
 
+	/** Is an interrupt in progress? */
+	volatile bool bIsInterrupt = false;
+
+	/** While the pipeline cache activities are thread safe and mostly lock-free,
+	 *  we may need to interrupt/lock the work to properly consolidate the pending caches. */
+	FRWLock InterruptLock;
+
 	FGraphEventRef RHICompletionEvent;
 
 	int32 ReleasedEntries = 0;
@@ -2659,6 +2725,27 @@ public:
 		return Result;
 	}
 
+	void GetResources(TArray<TRefCountPtr<FRHIResource>>& OutResources)
+	{
+		FScopeLock ScopeLock(&CriticalSection);
+
+		for (auto&& [Desc, State] : FullPipelines)
+		{
+			if (State)
+			{
+				OutResources.Add(TRefCountPtr<FRHIResource>(State->RHIPipeline));
+			}
+		}
+		
+		for (auto&& [Desc, State] : PartialPipelines)
+		{
+			if (State)
+			{
+				OutResources.Add(TRefCountPtr<FRHIResource>(State->RHIPipeline));
+			}
+		}
+	}
+
 	void Shutdown()
 	{
 		FScopeLock ScopeLock(&CriticalSection);
@@ -4348,3 +4435,13 @@ FRHIVertexDeclaration*	PipelineStateCache::GetOrCreateVertexDeclaration(const FV
 	GVertexDeclarationCache.Add(Key, NewDeclaration);
 	return NewDeclaration;
 }
+
+void PipelineStateCache::GetPipelineStates(TArray<TRefCountPtr<FRHIResource>>& Out, bool bConsolidateCaches)
+{
+	GComputePipelineCache.GetResources(Out, bConsolidateCaches);
+	GGraphicsPipelineCache.GetResources(Out, bConsolidateCaches);
+	GWorkGraphPipelineCache.GetResources(Out, bConsolidateCaches);
+#if RHI_RAYTRACING
+	GRayTracingPipelineCache.GetResources(Out);
+#endif // RHI_RAYTRACING
+}
diff --git a/Engine/Source/Runtime/RHI/Public/PipelineStateCache.h b/Engine/Source/Runtime/RHI/Public/PipelineStateCache.h
--- a/Engine/Source/Runtime/RHI/Public/PipelineStateCache.h
+++ b/Engine/Source/Runtime/RHI/Public/PipelineStateCache.h
@@ -130,6 +130,8 @@ namespace PipelineStateCache
 
 	extern RHI_API FGraphicsPipelineState*	FindGraphicsPipelineState(const FGraphicsPipelineStateInitializer& Initializer, bool bVerifyUse = true);
 
+	extern RHI_API void                     GetPipelineStates(TArray<TRefCountPtr<FRHIResource>>& Out, bool bConsolidateCaches = true);
+
 	extern RHI_API FRHIVertexDeclaration*	GetOrCreateVertexDeclaration(const FVertexDeclarationElementList& Elements);
 
 	// Retrieves RTPSO object from cache or adds a task to create it, which will be waited on by RHI thread.
diff --git a/Engine/Source/Runtime/RHICore/Internal/ShaderBinaryUtilities.h b/Engine/Source/Runtime/RHICore/Internal/ShaderBinaryUtilities.h
new file mode 100644
--- /dev/null
+++ b/Engine/Source/Runtime/RHICore/Internal/ShaderBinaryUtilities.h
@@ -0,0 +1,154 @@
+// Copyright Epic Games, Inc. All Rights Reserved.
+
+#pragma once
+
+#include "RHICore.h"
+#include "Math/Int128.h"
+
+namespace UE::ShaderBinaryUtilities
+{
+#if PLATFORM_SUPPORTS_PRAGMA_PACK
+#pragma pack(push, 1)
+#endif // PLATFORM_SUPPORTS_PRAGMA_PACK
+
+	namespace DXBC
+	{
+		struct FHeader
+		{
+			uint32 Identifier;
+			uint32 Ignore[6];
+			uint32 ChunkCount;
+		};
+
+		struct FChunkHeader
+		{
+			uint32 Type;
+			uint32 Size;
+		};
+	}
+	
+	namespace DXIL
+	{
+		struct FShaderDebugNameInfo
+		{
+			uint16 Flags;
+			uint16 NameLength;
+		};
+
+		struct FShaderHashInfo
+		{
+			uint32 Flags;
+			uint8  Digest[16];
+		};
+	}
+
+#if PLATFORM_SUPPORTS_PRAGMA_PACK
+#pragma pack(pop)
+#endif // PLATFORM_SUPPORTS_PRAGMA_PACK
+	
+	struct FParseContext
+	{
+		FParseContext(const void* Start, const uint64 ByteSize) : Start(static_cast<const uint8*>(Start)), ByteSize(ByteSize)
+		{
+			
+		}
+
+		template<typename T>
+		T Consume()
+		{
+			if (Offset + sizeof(T) > ByteSize)
+			{
+				checkf(false, TEXT("Parsing beyond EOS"));
+				return T{};
+			}
+			
+			const T* Value = reinterpret_cast<const T*>(Start + Offset);
+			Offset += sizeof(T);
+			return *Value;
+		}
+
+		FParseContext Split(uint64 InOffset) const
+		{
+			FParseContext ctx = *this;
+			ctx.Offset = InOffset;
+			return ctx;
+		}
+
+		uint64 PendingBytes() const
+		{
+			return ByteSize - Offset;
+		}
+
+		const void* Data() const
+		{
+			return Start + Offset;
+		}
+
+		const uint8* Start = nullptr;
+		const uint64 ByteSize = 0;
+		uint64 Offset = 0;
+	};
+	
+	static bool GetShaderBinaryDebugHashDXBC(const void* ShaderBinary, uint32 ByteSize, FUInt128& OutDebugHash)
+	{
+		if (ByteSize < sizeof(DXBC::FHeader))
+		{
+			UE_LOG(LogRHICore, Error, TEXT("Shader byte size too small"));
+			return false;
+		}
+
+		FParseContext Ctx(ShaderBinary, ByteSize);
+
+		auto Header = Ctx.Consume<DXBC::FHeader>();
+
+		for (uint32 ChunkIndex = 0; ChunkIndex < Header.ChunkCount; ++ChunkIndex)
+		{
+			uint32 ChunkOffset = Ctx.Consume<uint32>();
+
+			FParseContext ChunkCtx = Ctx.Split(ChunkOffset);
+
+			auto ChunkHeader = ChunkCtx.Consume<DXBC::FChunkHeader>();
+			switch (ChunkHeader.Type)
+			{
+				default:
+				{
+					continue;
+				}
+				case 'NDLI':
+				{
+					auto DXILDebugInfo = ChunkCtx.Consume<DXIL::FShaderDebugNameInfo>();
+
+					constexpr uint32 HashLength = 32;
+
+					if (DXILDebugInfo.NameLength != HashLength + FCString::Strlen(TEXT(".pdb")))
+					{
+						UE_LOG(LogRHICore, Display, TEXT("DXIL name length not the expected hash"));
+						return false;
+					}
+
+					if (ChunkCtx.PendingBytes() < HashLength)
+					{
+						UE_LOG(LogRHICore, Display, TEXT("ILDN block corrupt"));
+						return false;
+					}
+
+					String::HexToBytes(FUtf8StringView(static_cast<const UTF8CHAR*>(ChunkCtx.Data()), HashLength), reinterpret_cast<uint8*>(&OutDebugHash));
+					break;
+				}
+				case 'HSAH':
+				{
+					auto DXILDebugInfo = ChunkCtx.Consume<DXIL::FShaderHashInfo>();
+
+					static_assert(sizeof(OutDebugHash) == sizeof(DXILDebugInfo.Digest), "Unexpected digest size");
+					FMemory::Memcpy(&OutDebugHash, DXILDebugInfo.Digest, sizeof(OutDebugHash));
+					break;
+				}
+			}
+			
+			return true;
+		}
+
+		// No relevant chunk found
+		return false;
+	}
+}
diff --git a/Engine/Source/Runtime/RHICore/Private/RHICoreNvidiaAftermath.cpp b/Engine/Source/Runtime/RHICore/Private/RHICoreNvidiaAftermath.cpp
--- a/Engine/Source/Runtime/RHICore/Private/RHICoreNvidiaAftermath.cpp
+++ b/Engine/Source/Runtime/RHICore/Private/RHICoreNvidiaAftermath.cpp
@@ -6,17 +6,22 @@
 
 #include "RHI.h"
 #include "RHICore.h"
+#include "ShaderBinaryUtilities.h"
 #include "HAL/FileManager.h"
 #include "Misc/App.h"
 #include "Misc/Paths.h"
 #include "Misc/CommandLine.h"
 #include "HAL/PlatformCrashContext.h"
+#include "Misc/ScopeRWLock.h"
 
 #if PLATFORM_WINDOWS
 	#include "Windows/AllowWindowsPlatformTypes.h"
 #endif
 
 THIRD_PARTY_INCLUDES_START
+	#if PLATFORM_WINDOWS
+		#include "d3d12.h"
+	#endif // PLATFORM_WINDOWS
 	#include "GFSDK_Aftermath.h"
 	#include "GFSDK_Aftermath_GpuCrashdump.h"
 	#include "GFSDK_Aftermath_GpuCrashDumpDecoding.h"
@@ -81,6 +86,13 @@ namespace UE::RHICore::Nvidia::Aftermath
 		ECVF_ReadOnly
 	);
 
+	static TAutoConsoleVariable<int> CVarAftermath_ShaderRegistration(
+		TEXT("r.GPUCrashDebugging.Aftermath.ShaderRegistration"),
+		1,
+		TEXT("Enable registration of shaders and pipelines in Aftermath."),
+		ECVF_ReadOnly
+	);
+
 	static TAutoConsoleVariable<float> CVarAftermath_DumpShaderDebugInfo(
 		TEXT("r.GPUCrashDebugging.Aftermath.DumpShaderDebugInfo"),
 		0,
@@ -88,17 +100,37 @@ namespace UE::RHICore::Nvidia::Aftermath
 		ECVF_ReadOnly
 	);
 
-	static TAutoConsoleVariable<float> CVar_DumpWaitTime(
-		TEXT("r.GPUCrashDebugging.Aftermath.DumpWaitTime"),
+	static TAutoConsoleVariable<float> CVar_DumpProcessWaitTime(
+		TEXT("r.GPUCrashDebugging.Aftermath.DumpProcessWaitTime"),
 		10.0f,
 		TEXT("Amount of time (in seconds) to wait for Aftermath to finish processing GPU crash dumps."),
 		ECVF_Default
 	);
 
+	static TAutoConsoleVariable<float> CVar_DumpStartWaitTime(
+		TEXT("r.GPUCrashDebugging.Aftermath.DumpStartWaitTime"),
+		3.0f,
+		TEXT("Amount of time (in seconds) to wait for Aftermath to start processing GPU crash dumps."),
+		ECVF_Default
+	);
+
 	static void* DllHandle = nullptr;
 	static bool bEnabled = false;
 	static uint32 Flags = 0;
 	static FResolveMarkerFunc GResolveMarkerFunc;
+	static FAssociateShadersFunc GLateShaderAssociateFunc;
+	
+	struct FShaderRegistration
+	{
+		FCriticalSection         Lock;
+		TMap<uint64, FUInt128>   ShaderBinaryToDebugHash;
+	} static ShaderRegistration;
+
+	struct FCrashIntermediates
+	{
+		FRWLock                     Lock;
+		TMap<uint64, TArray<uint8>> ShaderDebugLookup;
+	} static CrashIntermediates;
 
 #if WITH_RHI_BREADCRUMBS
 	static TIndirectArray<FRHIBreadcrumb::FBuffer> NameStorage;
@@ -109,6 +141,7 @@ namespace UE::RHICore::Nvidia::Aftermath
 	static FCrashResult CrashResult {};
 
 	RHICORE_API bool IsEnabled() { return bEnabled; }
+	RHICORE_API bool IsShaderRegistrationEnabled() { return bEnabled && CVarAftermath_ShaderRegistration->GetInt(); }
 	RHICORE_API bool AreMarkersEnabled() { return bEnabled && (Flags & GFSDK_Aftermath_FeatureFlags_EnableMarkers); }
 
 	static void GFSDK_AFTERMATH_CALL Callback_GpuCrashDump    (const void* GPUCrashDumpData, const uint32_t GPUCrashDumpSize, void* UserData);
@@ -152,6 +185,11 @@ namespace UE::RHICore::Nvidia::Aftermath
 		}
 #endif
 	}
+	
+	RHICORE_API void SetLateShaderAssociateCallback(FAssociateShadersFunc ShaderAssociateFunc)
+	{
+		GLateShaderAssociateFunc = MoveTemp(ShaderAssociateFunc);
+	}
 
 	RHICORE_API void InitializeBeforeDeviceCreation(FResolveMarkerFunc ResolveMarkerFunc)
 	{
@@ -172,12 +210,11 @@ namespace UE::RHICore::Nvidia::Aftermath
 			;
 
 		Flags = GFSDK_Aftermath_FeatureFlags_Minimum;
-		Flags |= (FParse::Param(FCommandLine::Get(), TEXT("nvaftermathmarkers"  )) || (CVarAftermath_Markers  ->GetInt())) ? GFSDK_Aftermath_FeatureFlags_EnableMarkers          : 0;
-		Flags |= (FParse::Param(FCommandLine::Get(), TEXT("nvaftermathresources")) || (CVarAftermath_Resources->GetInt())) ? GFSDK_Aftermath_FeatureFlags_EnableResourceTracking : 0;
-		Flags |= (FParse::Param(FCommandLine::Get(), TEXT("nvaftermathcallstack")) || (CVarAftermath_Callstack->GetInt())) ? GFSDK_Aftermath_FeatureFlags_CallStackCapturing     : 0;
-		Flags |= (FParse::Param(FCommandLine::Get(), TEXT("nvaftermathall"      )) || (CVarAftermath_TrackAll ->GetInt())) ? AllFlags                                            : 0;
-
-		const bool bDumpShaderDebugInfo = (FParse::Param(FCommandLine::Get(), TEXT("nvAftermathDumpShaderDebugInfo")) || (CVarAftermath_DumpShaderDebugInfo->GetInt()));
+		Flags |= (FParse::Param(FCommandLine::Get(), TEXT("nvaftermathmarkers"            )) || (CVarAftermath_Markers            ->GetInt())) ? GFSDK_Aftermath_FeatureFlags_EnableMarkers           : 0;
+		Flags |= (FParse::Param(FCommandLine::Get(), TEXT("nvaftermathresources"          )) || (CVarAftermath_Resources          ->GetInt())) ? GFSDK_Aftermath_FeatureFlags_EnableResourceTracking  : 0;
+		Flags |= (FParse::Param(FCommandLine::Get(), TEXT("nvaftermathcallstack"          )) || (CVarAftermath_Callstack          ->GetInt())) ? GFSDK_Aftermath_FeatureFlags_CallStackCapturing      : 0;
+		Flags |= (FParse::Param(FCommandLine::Get(), TEXT("nvAftermathDumpShaderDebugInfo")) || (CVarAftermath_DumpShaderDebugInfo->GetInt())) ? GFSDK_Aftermath_FeatureFlags_GenerateShaderDebugInfo : 0;
+		Flags |= (FParse::Param(FCommandLine::Get(), TEXT("nvaftermathall"                )) || (CVarAftermath_TrackAll           ->GetInt())) ? AllFlags                                             : 0;
 
 		const bool bResolveMarkers = NV_AFTERMATH_USE_BREADCRUMB_PTRS || GResolveMarkerFunc.IsSet();
 		
@@ -190,7 +227,7 @@ namespace UE::RHICore::Nvidia::Aftermath
 #endif
 			GFSDK_Aftermath_GpuCrashDumpFeatureFlags_DeferDebugInfoCallbacks,
 			&Callback_GpuCrashDump,
-			bDumpShaderDebugInfo ? &Callback_ShaderDebugInfo : nullptr,
+			(Flags & GFSDK_Aftermath_FeatureFlags_GenerateShaderDebugInfo) ? &Callback_ShaderDebugInfo : nullptr,
 			&Callback_GpuCrashDumpDesc,
 			bResolveMarkers ? &Callback_ResolveMarker : nullptr,
 			nullptr // user data
@@ -256,6 +293,152 @@ namespace UE::RHICore::Nvidia::Aftermath
 		return true;
 	}
 
+	static void AppendGpuCrashShaderBinaryInfo(GFSDK_Aftermath_GpuCrashDump_Decoder Decoder, const GFSDK_Aftermath_GpuCrashDump_ShaderInfo& Shader)
+	{
+		GFSDK_Aftermath_ShaderBinaryHash BinaryHash;
+
+		GFSDK_Aftermath_Result Result = GFSDK_Aftermath_GetShaderHashForShaderInfo(Decoder, &Shader, &BinaryHash);
+		if (Result != GFSDK_Aftermath_Result_Success)
+		{
+			CrashResult.OutputLog += FString::Printf(TEXT("\n\t\t\t! Failed to get binary hash (%u)"), static_cast<uint32>(Result));
+			return;
+		}
+
+		// Report binary hash separately, may not be the same
+		CrashResult.OutputLog += FString::Printf(TEXT("\n\t\t\tBinary Hash = %llu"), BinaryHash.hash);
+
+		FUInt128* Candidate = ShaderRegistration.ShaderBinaryToDebugHash.Find(BinaryHash.hash);
+		if (!Candidate)
+		{
+			CrashResult.OutputLog += TEXT("\n\t\t\t! Failed to associate Aftermath hash");
+			return;
+		}
+
+		// Report debug association
+		FString Hex = BytesToHex(reinterpret_cast<const uint8*>(&Candidate), sizeof(FUInt128));
+		CrashResult.OutputLog += FString::Printf(TEXT("\n\t\t\tPDB = %s.pdb"), *Hex);
+	}
+
+	static void AppendGpuCrashShaderInfo(GFSDK_Aftermath_GpuCrashDump_Decoder Decoder)
+	{
+		uint32 ShaderCount = 0;
+
+		FScopeLock Lock(&ShaderRegistration.Lock);
+		
+		// Total number of faulting shaders
+		GFSDK_Aftermath_Result Result = GFSDK_Aftermath_GpuCrashDump_GetActiveShadersInfoCount(Decoder, &ShaderCount);
+		if (Result != GFSDK_Aftermath_Result_Success)
+		{
+			CrashResult.OutputLog += FString::Printf(TEXT("\n\t\tFailed to get shader count (%u)."), static_cast<uint32>(Result));
+			return;
+		}
+		
+		if (!ShaderCount)
+		{
+			CrashResult.OutputLog += TEXT("\n\t\tNo Shaders.");
+			return;
+		}
+		
+		TArray<GFSDK_Aftermath_GpuCrashDump_ShaderInfo> Shaders;
+		Shaders.SetNumUninitialized(ShaderCount);
+
+		// Get all faulting shaders
+		Result = GFSDK_Aftermath_GpuCrashDump_GetActiveShadersInfo(Decoder, ShaderCount, Shaders.GetData());
+		if (Result != GFSDK_Aftermath_Result_Success)
+		{
+			CrashResult.OutputLog += FString::Printf(TEXT("\n\t\tFailed to get shaders (%u)."), static_cast<uint32>(Result));
+			return;
+		}
+
+		CrashResult.OutputLog += FString::Printf(TEXT("\n\t\t%u total."), ShaderCount);
+
+		for (uint32 ShaderIndex = 0; ShaderIndex < ShaderCount; ++ShaderIndex)
+		{
+			CrashResult.OutputLog += FString::Printf(TEXT("\n\t\t[%u]:"), ShaderIndex);
+
+			const GFSDK_Aftermath_GpuCrashDump_ShaderInfo& Shader = Shaders[ShaderIndex];
+			
+			if (Shader.isInternal)
+			{
+				CrashResult.OutputLog += TEXT("\n\t\t\t! Internal");
+			}
+
+			// Translate faulting shader type
+			const TCHAR* ShaderType = nullptr;
+			switch (Shader.shaderType)
+			{
+			default:
+			case GFSDK_Aftermath_ShaderType_Unknown:
+				ShaderType = TEXT("Unknown");
+				break;
+			case GFSDK_Aftermath_ShaderType_Vertex:
+				ShaderType = TEXT("Vertex");
+				break;
+			case GFSDK_Aftermath_ShaderType_Tessellation_Control:
+				ShaderType = TEXT("Tessellation Control / Hull");
+				break;
+			case GFSDK_Aftermath_ShaderType_Tessellation_Evaluation:
+				ShaderType = TEXT("Tessellation Evaluation / Domain");
+				break;
+			case GFSDK_Aftermath_ShaderType_Geometry:
+				ShaderType = TEXT("Geometry");
+				break;
+			case GFSDK_Aftermath_ShaderType_Fragment:
+				ShaderType = TEXT("Pixel");
+				break;
+			case GFSDK_Aftermath_ShaderType_Compute:
+				ShaderType = TEXT("Compute");
+				break;
+			case GFSDK_Aftermath_ShaderType_RayTracing_RayGeneration:
+				ShaderType = TEXT("RayTracing Generation");
+				break;
+			case GFSDK_Aftermath_ShaderType_RayTracing_Miss:
+				ShaderType = TEXT("RayTracing Miss");
+				break;
+			case GFSDK_Aftermath_ShaderType_RayTracing_Intersection:
+				ShaderType = TEXT("RayTracing Intersection");
+				break;
+			case GFSDK_Aftermath_ShaderType_RayTracing_AnyHit:
+				ShaderType = TEXT("RayTracing AnyHit");
+				break;
+			case GFSDK_Aftermath_ShaderType_RayTracing_ClosestHit:
+				ShaderType = TEXT("RayTracing ClosestHit");
+				break;
+			case GFSDK_Aftermath_ShaderType_RayTracing_Callable:
+				ShaderType = TEXT("RayTracing Callable");
+				break;
+			case GFSDK_Aftermath_ShaderType_RayTracing_Internal:
+				ShaderType = TEXT("RayTracing Internal");
+				break;
+			case GFSDK_Aftermath_ShaderType_Mesh:
+				ShaderType = TEXT("Mesh");
+				break;
+			case GFSDK_Aftermath_ShaderType_Task:
+				ShaderType = TEXT("Task");
+				break;
+			}
+			
+			CrashResult.OutputLog += FString::Printf(TEXT("\n\t\t\tType = %s"), ShaderType);
+			CrashResult.OutputLog += FString::Printf(TEXT("\n\t\t\tHash = %llu"), Shader.shaderHash);
+
+			AppendGpuCrashShaderBinaryInfo(Decoder, Shader);
+		}
+	}
+
+	static void GFSDK_AFTERMATH_CALL Callback_ShaderDebugInfoLookup(const GFSDK_Aftermath_ShaderDebugInfoIdentifier* Identifier, PFN_GFSDK_Aftermath_SetData SetShaderDebugInfo, void* UserData)
+	{
+		FRWScopeLock Lock(CrashIntermediates.Lock, SLT_ReadOnly);
+
+		TArray<uint8>* DebugData = CrashIntermediates.ShaderDebugLookup.Find(Identifier->id[0]);
+		if (!DebugData)
+		{
+			UE_LOG(LogRHICore, Error, TEXT("Shader debug data not found (%llu)"), Identifier->id[0]);
+			return;
+		}
+		
+		SetShaderDebugInfo(DebugData->GetData(), DebugData->Num());
+	}
+	
 	static void GFSDK_AFTERMATH_CALL Callback_GpuCrashDump(const void* GPUCrashDumpData, const uint32_t GPUCrashDumpSize, void* UserData)
 	{
 		ON_SCOPE_EXIT
@@ -268,7 +451,8 @@ namespace UE::RHICore::Nvidia::Aftermath
 
 		// Write out crash dump to project log dir
 		CrashResult.DumpPath = FPaths::Combine(FPaths::ProjectLogDir(), FString::Printf(TEXT("%s.%s.nv-gpudmp"), GDynamicRHI->GetNonValidationRHI()->GetName(), *FDateTime::Now().ToString()));
-
+		FGenericCrashContext::AddFile(*CrashResult.DumpPath);
+		
 		CrashResult.OutputLog += FString::Printf(TEXT("\n\nWriting Aftermath dump to: %s"), *CrashResult.DumpPath.GetValue());
 		if (FArchive* Writer = IFileManager::Get().CreateFileWriter(*CrashResult.DumpPath.GetValue()))
 		{
@@ -555,6 +739,18 @@ namespace UE::RHICore::Nvidia::Aftermath
 			}
 		}
 
+		// Append shader information
+		CrashResult.OutputLog += TEXT("\n\n\tActive Shaders:");
+		{
+			// Allow association after a fault
+			if (GLateShaderAssociateFunc)
+			{
+				GLateShaderAssociateFunc();
+			}
+			
+			AppendGpuCrashShaderInfo(Decoder);
+		}
+
 		// Append JSON dump generated by the decoder.
 		CrashResult.OutputLog += TEXT("\n\n\tJSON Data:");
 		{
@@ -563,12 +759,15 @@ namespace UE::RHICore::Nvidia::Aftermath
 				Decoder,
 				GFSDK_Aftermath_GpuCrashDumpDecoderFlags_ALL_INFO,
 				GFSDK_Aftermath_GpuCrashDumpFormatterFlags_UTF8_OUTPUT,
-				nullptr,
+				Callback_ShaderDebugInfoLookup,
 				nullptr,
 				nullptr,
 				nullptr,
 				&JsonSize);
 
+			// Cleanup
+			CrashIntermediates.ShaderDebugLookup.Empty();
+
 			if (Result != GFSDK_Aftermath_Result_Success)
 			{
 				CrashResult.OutputLog += FString::Printf(TEXT("\n\t\tFailed to generate JSON from GPU crash dump: 0x%08x"), Result);
@@ -652,7 +851,76 @@ namespace UE::RHICore::Nvidia::Aftermath
 				Writer->Serialize((void*)ShaderDebugInfo, ShaderDebugInfoSize);
 				Writer->Close();
 			}
+
+			// Make sure the crash handler picks up the file
+			FGenericCrashContext::AddFile(Filename);
+
+			// Keep around for json dumping
+			FRWScopeLock Lock(CrashIntermediates.Lock, SLT_Write);
+			CrashIntermediates.ShaderDebugLookup.Add(Identifier.id[0], TArray(static_cast<const uint8*>(ShaderDebugInfo), ShaderDebugInfoSize));
+		}
+	}
+
+	FShaderHash RegisterShaderBinary(const void* ShaderBinary, uint32 ByteSize)
+	{
+		if (ByteSize < sizeof(uint32))
+		{
+			UE_LOG(LogRHICore, Display, TEXT("Shader byte size too small, expected magic dword"));
+			return InvalidShaderHash;
+		}
+
+		// On disk debug hash, typically the [hash].pdb
+		FUInt128 DebugHash;
+
+		// Internal aftermath hash
+		GFSDK_Aftermath_ShaderBinaryHash AftermathHash{};
+		
+		switch (*static_cast<const uint32*>(ShaderBinary))
+		{
+			default:
+			{
+				UE_LOG(LogRHICore, Display, TEXT("Unknown shader format"));
+				return InvalidShaderHash;
+			}
+			case 'CBXD': // DXBC container header
+			{
+				ShaderBinaryUtilities::GetShaderBinaryDebugHashDXBC(ShaderBinary, ByteSize, DebugHash);
+					
+#ifdef __d3d12_h__
+				D3D12_SHADER_BYTECODE ByteCode;
+				ByteCode.pShaderBytecode = ShaderBinary;
+				ByteCode.BytecodeLength = ByteSize;
+
+				GFSDK_Aftermath_Result Result = GFSDK_Aftermath_GetShaderHash(GFSDK_Aftermath_Version_API, &ByteCode, &AftermathHash);
+				if (Result != GFSDK_Aftermath_Result_Success)
+				{
+					UE_LOG(LogNvidiaAftermath, Display, TEXT("Failed to get aftermath shader hash (%u)"), Result);
+					return InvalidShaderHash;
+				}
+#else // __d3d12_h__
+				UE_LOG(LogNvidiaAftermath, Display, TEXT("DXBC container without D3D12 support"), Result);
+				return InvalidShaderHash;
+#endif // __d3d12_h__
+					
+				break;
+			}
+			case 0x07230203: // Vulkan magic header
+			{
+				UE_LOG(LogRHICore, Display, TEXT("SPIRV binary debug hash not supported yet"));
+				return InvalidShaderHash;
+			}
 		}
+
+		FScopeLock Lock(&ShaderRegistration.Lock);
+		ShaderRegistration.ShaderBinaryToDebugHash.Add(AftermathHash.hash, DebugHash);
+		
+		return AftermathHash.hash;
+	}
+	
+	void DeregisterShaderBinary(FShaderHash Hash)
+	{
+		FScopeLock Lock(&ShaderRegistration.Lock);
+		ShaderRegistration.ShaderBinaryToDebugHash.Remove(Hash);
 	}
 
 	FCrashResult OnGPUCrash()
@@ -661,7 +929,8 @@ namespace UE::RHICore::Nvidia::Aftermath
 			return {};
 
 		const double StartTime = FPlatformTime::Seconds();
-		const double EndTime = StartTime + CVar_DumpWaitTime->GetFloat();
+		const double ProcessingEndTime = StartTime + CVar_DumpProcessWaitTime->GetFloat();
+		const double NotStartedEndTime = StartTime + CVar_DumpStartWaitTime->GetFloat();
 
 		// Wait for GPU crash dump processing to complete.
 		while (true)
@@ -681,15 +950,22 @@ namespace UE::RHICore::Nvidia::Aftermath
 			case GFSDK_Aftermath_CrashDump_Status_Unknown:
 				UE_LOG(LogNvidiaAftermath, Warning, TEXT("Unknown aftermath crash dump state. Skipping crash dump processing."));
 				return {};
-
 			case GFSDK_Aftermath_CrashDump_Status_NotStarted:
-				UE_LOG(LogNvidiaAftermath, Warning, TEXT("Aftermath did not detect a GPU crash. Skipping crash dump processing."));
-				return {};
-
+				// Driver may take time before starting the dump
+				if (FPlatformTime::Seconds() >= NotStartedEndTime)
+				{
+					UE_LOG(LogNvidiaAftermath, Warning, TEXT("Timed out while waiting for Aftermath to start the GPU crash dump."));
+					return {};
+				}
+				else
+				{
+					FPlatformProcess::Sleep(0.01f);
+				}
+				break;
 			case GFSDK_Aftermath_CrashDump_Status_CollectingData:
 			case GFSDK_Aftermath_CrashDump_Status_InvokingCallback:
 				// Crash dump is still in progress...
-				if (FPlatformTime::Seconds() >= EndTime)
+				if (FPlatformTime::Seconds() >= ProcessingEndTime)
 				{
 					UE_LOG(LogNvidiaAftermath, Warning, TEXT("Timed out while waiting for Aftermath to complete GPU crash dump."));
 					return {};
diff --git a/Engine/Source/Runtime/RHICore/Public/RHICoreNvidiaAftermath.h b/Engine/Source/Runtime/RHICore/Public/RHICoreNvidiaAftermath.h
--- a/Engine/Source/Runtime/RHICore/Public/RHICoreNvidiaAftermath.h
+++ b/Engine/Source/Runtime/RHICore/Public/RHICoreNvidiaAftermath.h
@@ -19,8 +19,10 @@
 	namespace UE::RHICore::Nvidia::Aftermath
 	{
 		typedef TUniqueFunction<void(const void* MarkerData, const uint32_t MarkerDataSize, void* UserData, void** ResolvedMarkerData, uint32_t* ResolvedMarkerDataSize)> FResolveMarkerFunc;
-
+		typedef TUniqueFunction<void()> FAssociateShadersFunc;
+		
 		RHICORE_API bool IsEnabled();
+		RHICORE_API bool IsShaderRegistrationEnabled();
 		RHICORE_API bool AreMarkersEnabled();
 
 		//
@@ -29,12 +31,33 @@
 		//
 		RHICORE_API void InitializeBeforeDeviceCreation(FResolveMarkerFunc ResolveMarkerFunc = {});
 
+		//
+		// Called by the crash handler to associate shaders after a GPU fault
+		//
+		RHICORE_API void SetLateShaderAssociateCallback(FAssociateShadersFunc ShaderAssociateFunc = {});
+
 		//
 		// Called by platform RHIs after device creation.
 		// InitCallback should return an Aftermath result value.
 		//
 		RHICORE_API bool InitializeDevice(TFunctionRef<uint32(uint32 AftermathFeatureFlags)> InitCallback);
 
+		using FShaderHash = uint64;
+		
+		static constexpr FShaderHash InvalidShaderHash = UINT64_MAX;
+
+		//
+		// Register a new shader binary with Aftermath, must be valid.
+		// Returns the internal shader hash, must be deregistered with DeregisterShaderBinary
+		//
+		RHICORE_API FShaderHash RegisterShaderBinary(const void* ShaderBinary, uint32 ByteSize);
+
+		//
+		// Deregister a shader binary from its hash. Will not release it immediately,
+		// but instead enqueue it for later eviction once the GPU work is done.
+		//
+		RHICORE_API void DeregisterShaderBinary(FShaderHash Hash);
+
 		struct FCrashResult
 		{
 			FString OutputLog;
