Check raytracing/workgraph support when pipeline gathering in Aftermath

diff --git a/Engine/Source/Runtime/D3D12RHI/Private/D3D12NvidiaExtensions.cpp b/Engine/Source/Runtime/D3D12RHI/Private/D3D12NvidiaExtensions.cpp
--- a/Engine/Source/Runtime/D3D12RHI/Private/D3D12NvidiaExtensions.cpp
+++ b/Engine/Source/Runtime/D3D12RHI/Private/D3D12NvidiaExtensions.cpp
@@ -128,6 +128,7 @@ namespace UE::RHICore::Nvidia::Aftermath::D3D12
 					ConditionallyAddShader(static_cast<FD3D12ComputeShader*>(Pipeline->GetComputeShader()), Shaders, ShaderSet);
 					break;
 				}
+#if D3D12_RHI_RAYTRACING
 				case RRT_RayTracingPipelineState:
 				{
 					auto* Pipeline = static_cast<FD3D12RayTracingPipelineState*>(Resource);
@@ -156,6 +157,8 @@ namespace UE::RHICore::Nvidia::Aftermath::D3D12
 					}
 					break;
 				}
+#endif // D3D12_RHI_RAYTRACING
+#if D3D12_RHI_WORKGRAPHS
 				case RRT_WorkGraphPipelineState:
 				{
 					auto* Pipeline = static_cast<FD3D12WorkGraphPipelineState*>(Resource);
@@ -173,6 +176,7 @@ namespace UE::RHICore::Nvidia::Aftermath::D3D12
 					}
 					break;
 				}
+#endif // D3D12_RHI_WORKGRAPHS
 			}
 		}
 
