Fixed incorrect Aftermath work graph shader type cast

diff --git a/Engine/Source/Runtime/D3D12RHI/Private/D3D12NvidiaExtensions.cpp b/Engine/Source/Runtime/D3D12RHI/Private/D3D12NvidiaExtensions.cpp
--- a/Engine/Source/Runtime/D3D12RHI/Private/D3D12NvidiaExtensions.cpp
+++ b/Engine/Source/Runtime/D3D12RHI/Private/D3D12NvidiaExtensions.cpp
@@ -171,7 +171,7 @@ namespace UE::RHICore::Nvidia::Aftermath::D3D12
 						
 					for (FRHIShader* Shader : Pipeline->Shaders)
 					{
-						auto* ShaderTyped = static_cast<FD3D12RayTracingShader*>(Shader);
+						auto* ShaderTyped = static_cast<FD3D12WorkGraphShader*>(Shader);
 						ConditionallyAddShader(ShaderTyped, Shaders, ShaderSet);
 					}
 					break;
