Use bindless code path for inline raytracing on d3d12 win64 rhi by using explicit descriptor heaps

diff --git a/Engine/Source/Runtime/D3D12RHI/Private/D3D12Adapter.cpp b/Engine/Source/Runtime/D3D12RHI/Private/D3D12Adapter.cpp
--- a/Engine/Source/Runtime/D3D12RHI/Private/D3D12Adapter.cpp
+++ b/Engine/Source/Runtime/D3D12RHI/Private/D3D12Adapter.cpp
@@ -1045,6 +1045,9 @@ void FD3D12Adapter::InitializeDevices()
 
 							GRHISupportsRayTracingPSOAdditions = true;
 							GRHISupportsInlineRayTracing = GRHISupportsRayTracing && RHISupportsInlineRayTracing(GMaxRHIShaderPlatform) && (GMaxRHIFeatureLevel >= ERHIFeatureLevel::SM6);
+														
+							// Inline RayTracing SBT is required
+							GRHIGlobals.RayTracing.RequiresInlineRayTracingSBT = GRHISupportsInlineRayTracing;
 						}
 						else
 						{
diff --git a/Engine/Source/Runtime/D3D12RHI/Private/D3D12Commands.cpp b/Engine/Source/Runtime/D3D12RHI/Private/D3D12Commands.cpp
--- a/Engine/Source/Runtime/D3D12RHI/Private/D3D12Commands.cpp
+++ b/Engine/Source/Runtime/D3D12RHI/Private/D3D12Commands.cpp
@@ -135,6 +135,8 @@ FD3D12ResourceLocation& FD3D12CommandContext::SetupIndirectArgument(FRHIBuffer*
 
 void FD3D12CommandContext::PostGpuEvent()
 {
+	UnsetExplicitDescriptorCache();
+
 	ConditionalSplitCommandList();
 	DEBUG_EXECUTE_COMMAND_LIST(this);
 }
diff --git a/Engine/Source/Runtime/D3D12RHI/Private/D3D12StateCache.cpp b/Engine/Source/Runtime/D3D12RHI/Private/D3D12StateCache.cpp
--- a/Engine/Source/Runtime/D3D12RHI/Private/D3D12StateCache.cpp
+++ b/Engine/Source/Runtime/D3D12RHI/Private/D3D12StateCache.cpp
@@ -4,6 +4,7 @@
 //	thread performance by removing redundant device context calls.
 
 #include "D3D12RHIPrivate.h"
+#include "D3D12ExplicitDescriptorCache.h"
 
 // This value defines how many descriptors will be in the device global descriptor heap. This heap contains all shader visible view descriptors.
 // Other shader visible descriptor heaps (e.g. OnlineViewHeap) are allocated from this pool. Non-visible heaps (e.g. LocalViewHeap) are allocated as standalone.
@@ -456,11 +457,6 @@ void FD3D12StateCache::ApplyState(ERHIPipeline HardwarePipe, ED3D12PipelineType
 		return;
 	}
 
-	const bool bRootSignatureChanged = InternalSetRootSignature(PipelineType, PSOCommonData->RootSignature);
-
-	// Ensure the correct graphics PSO is set.
-	InternalSetPipelineState(PSOCommonData->PipelineState);
-
 	bool bBindlessResources = PSOCommonData->RootSignature->UsesDynamicResources();
 	bool bBindlessSamplers = PSOCommonData->RootSignature->UsesDynamicSamplers();
 
@@ -477,14 +473,34 @@ void FD3D12StateCache::ApplyState(ERHIPipeline HardwarePipe, ED3D12PipelineType
 		bool bHaveResourceHeap = BindlessManager.AreResourcesBindless(ERHIBindlessConfiguration::AllShaders);
 		bool bHaveSamplerHeap = BindlessManager.AreSamplersBindless(ERHIBindlessConfiguration::AllShaders);
 
-		checkf(!bBindlessResources || bHaveResourceHeap, TEXT("Using dynamic samplers without the bindless sampler heap configured. Please check your configuration."));
-		checkf(!bBindlessSamplers  || bHaveSamplerHeap, TEXT("Using dynamic samplers without the bindless sampler heap configured. Please check your configuration."));
+		bool bRequireExplicitDescriptorCacheSet = bBindlessResources && !bHaveResourceHeap;
+		if (bRequireExplicitDescriptorCacheSet)
+		{
+			// Samplers should be bindless as well
+			check(bBindlessSamplers);
+			// Should have any SRVs or UAVs
+			check(!bHasTableResources);
+
+			// Everything bindless, don't require any CVS, SRV and UAV view data
+			FD3D12ExplicitDescriptorCache TransientDescriptorCache(GetParentDevice(), 1);
+			TransientDescriptorCache.Init(0, 0, 0, ERHIBindlessConfiguration::RayTracingShaders);
+
+			CmdContext.SetExplicitDescriptorCache(TransientDescriptorCache);
+			
+			PipelineState.Compute.bNeedSetRootSignature = true;
+			PipelineState.Graphics.bNeedSetRootSignature = true;
+		}
 
 		check(!(bHaveResourceHeap && bHasTableResources));
 		check(!(bHaveSamplerHeap  && bHasSamplers      ));
 	}
 #endif
 
+	const bool bRootSignatureChanged = InternalSetRootSignature(PipelineType, PSOCommonData->RootSignature);
+
+	// Ensure the correct graphics PSO is set.
+	InternalSetPipelineState(PSOCommonData->PipelineState);
+
 	if (bRootSignatureChanged)
 	{
 		CmdContext.BindDiagnosticBuffer(PSOCommonData->RootSignature, PipelineType);
diff --git a/Engine/Source/Runtime/Engine/Private/ShaderCompiler/ShaderCompiler.cpp b/Engine/Source/Runtime/Engine/Private/ShaderCompiler/ShaderCompiler.cpp
--- a/Engine/Source/Runtime/Engine/Private/ShaderCompiler/ShaderCompiler.cpp
+++ b/Engine/Source/Runtime/Engine/Private/ShaderCompiler/ShaderCompiler.cpp
@@ -3994,7 +3994,7 @@ void GlobalBeginCompileShader(
 	bool bEnableBindlessMacro = false;
 	if (RHIGetBindlessSupport(ShaderPlatform) != ERHIBindlessSupport::Unsupported && !Input.Environment.CompilerFlags.Contains(CFLAG_ForceBindful))
 	{
-		const bool bIsRaytracingShader = IsRayTracingShaderFrequency(Input.Target.GetFrequency());
+		const bool bIsRaytracingShader = IsRayTracingShaderFrequency(Input.Target.GetFrequency()) || Input.Environment.CompilerFlags.Contains(CFLAG_InlineRayTracing);
 
 		const ERHIBindlessConfiguration ResourcesConfig = UE::ShaderCompiler::GetBindlessResourcesConfiguration(ShaderFormatName);
 		const ERHIBindlessConfiguration SamplersConfig = UE::ShaderCompiler::GetBindlessSamplersConfiguration(ShaderFormatName);
