[HWRT] Always enable NVAPI extension in RTPSOs when IsRHIDeviceNVIDIA()

diff --git a/Engine/Shaders/Private/PathTracing/PathTracingCore.ush b/Engine/Shaders/Private/PathTracing/PathTracingCore.ush
--- a/Engine/Shaders/Private/PathTracing/PathTracingCore.ush
+++ b/Engine/Shaders/Private/PathTracing/PathTracingCore.ush
@@ -41,9 +41,6 @@
 
 #if PATH_TRACER_USE_SER
 
-#define NV_SHADER_EXTN_SLOT u0
-#define NV_SHADER_EXTN_REGISTER_SPACE space1001
-
 #define NV_HITOBJECT_USE_MACRO_API
 #include "/Engine/Shared/ThirdParty/NVIDIA/nvHLSLExtns.h"
 
diff --git a/Engine/Shaders/Shared/HLSLReservedSpaces.h b/Engine/Shaders/Shared/HLSLReservedSpaces.h
--- a/Engine/Shaders/Shared/HLSLReservedSpaces.h
+++ b/Engine/Shaders/Shared/HLSLReservedSpaces.h
@@ -47,4 +47,11 @@
 // Register space for UE diagnotic debug buffer
 #define UE_HLSL_SPACE_DIAGNOSTIC 999
 
-// NOTE: AMD AGS uses AGS_DX12_SHADER_INSTRINSICS_SPACE_ID=2147420894
\ No newline at end of file
+// NOTE: AMD AGS uses AGS_DX12_SHADER_INSTRINSICS_SPACE_ID=2147420894
+
+// Nvidia Shader Extensions
+#define UE_HLSL_SLOT_NV_SHADER_EXTN 0
+#define UE_HLSL_SPACE_NV_SHADER_EXTN 1001
+
+#define NV_SHADER_EXTN_SLOT UE_REG_CONCATENATE(u, UE_HLSL_SLOT_NV_SHADER_EXTN)
+#define NV_SHADER_EXTN_REGISTER_SPACE UE_REG_CONCATENATE(space, UE_HLSL_SPACE_NV_SHADER_EXTN)
diff --git a/Engine/Source/Runtime/D3D12RHI/Private/D3D12NvidiaExtensions.h b/Engine/Source/Runtime/D3D12RHI/Private/D3D12NvidiaExtensions.h
--- a/Engine/Source/Runtime/D3D12RHI/Private/D3D12NvidiaExtensions.h
+++ b/Engine/Source/Runtime/D3D12RHI/Private/D3D12NvidiaExtensions.h
@@ -14,12 +14,6 @@
 	#include "Microsoft/HideMicrosoftPlatformTypes.h"
 #endif // WITH_NVAPI
 
-#if WITH_NVAPI
-	// NVidia Shader Execution Reordering extension
-	#define NV_SER_EXTENSION_SLOT  0
-	#define NV_SER_EXTENSION_SPACE 1001
-#endif
-
 #if NV_AFTERMATH
 
 	struct GFSDK_Aftermath_ContextHandle__;
diff --git a/Engine/Source/Runtime/D3D12RHI/Private/D3D12RayTracing.cpp b/Engine/Source/Runtime/D3D12RHI/Private/D3D12RayTracing.cpp
--- a/Engine/Source/Runtime/D3D12RHI/Private/D3D12RayTracing.cpp
+++ b/Engine/Source/Runtime/D3D12RHI/Private/D3D12RayTracing.cpp
@@ -644,10 +644,10 @@ static TRefCountPtr<ID3D12StateObject> CreateRayTracingStateObject(
 
 #if WITH_NVAPI
 	bool bReorderingDisabled = false;
-	if (GRHIGlobals.SupportsShaderExecutionReordering && IsRHIDeviceNVIDIA())
+	if (IsRHIDeviceNVIDIA())
 	{
 		// Enable the NVAPI extension to allow shader extensions in ray tracing for shader execution reordering
-		NvAPI_Status NvapiStatus = NvAPI_D3D12_SetNvShaderExtnSlotSpaceLocalThread(RayTracingDevice, NV_SER_EXTENSION_SLOT, NV_SER_EXTENSION_SPACE);
+		NvAPI_Status NvapiStatus = NvAPI_D3D12_SetNvShaderExtnSlotSpaceLocalThread(RayTracingDevice, UE_HLSL_SLOT_NV_SHADER_EXTN, UE_HLSL_SPACE_NV_SHADER_EXTN);
 		bReorderingDisabled = ensureMsgf(NvapiStatus == NVAPI_OK, TEXT("NVAPI ERROR %d on setup extensions\n"), NvapiStatus);
 	}
 #endif
diff --git a/Engine/Source/Runtime/D3D12RHI/Private/D3D12RootSignature.cpp b/Engine/Source/Runtime/D3D12RHI/Private/D3D12RootSignature.cpp
--- a/Engine/Source/Runtime/D3D12RHI/Private/D3D12RootSignature.cpp
+++ b/Engine/Source/Runtime/D3D12RHI/Private/D3D12RootSignature.cpp
@@ -435,12 +435,12 @@ FD3D12RootSignatureDesc::FD3D12RootSignatureDesc(const FD3D12QuantizedBoundShade
 	}
 
 #if WITH_NVAPI
-	if (QBSS.RootSignatureType == RS_RayTracingGlobal && GRHIGlobals.SupportsShaderExecutionReordering && IsRHIDeviceNVIDIA())
+	if (QBSS.RootSignatureType == RS_RayTracingGlobal && IsRHIDeviceNVIDIA())
 	{
 		check(RootParameterCount < MaxRootParameters);
 
 		// Should this use the InitAsUnorderedAccessView???
-		DescriptorRanges[RootParameterCount].Init(D3D12_DESCRIPTOR_RANGE_TYPE_UAV, 1, NV_SER_EXTENSION_SLOT, NV_SER_EXTENSION_SPACE, D3D12_DESCRIPTOR_RANGE_FLAG_NONE, 0);
+		DescriptorRanges[RootParameterCount].Init(D3D12_DESCRIPTOR_RANGE_TYPE_UAV, 1, UE_HLSL_SLOT_NV_SHADER_EXTN, UE_HLSL_SPACE_NV_SHADER_EXTN, D3D12_DESCRIPTOR_RANGE_FLAG_NONE, 0);
 		TableSlots[RootParameterCount].InitAsDescriptorTable(1, &DescriptorRanges[RootParameterCount]);
 		RootParameterCount++;
 		RootParametersSize += RootDescriptorCost;
