Resubmit of inline raytracing for d3d12 pc with fixes:
- add ddpi for inline raytracing requiring bindless (enabled for d3d12 pc only)
- only enable bindless if inline raytracing and bindless is required in shader compiler setup
- Use bindless code path for inline raytracing on d3d12 win64 rhi by using explicit descriptor heaps

diff --git a/Engine/Config/Windows/DataDrivenPlatformInfo.ini b/Engine/Config/Windows/DataDrivenPlatformInfo.ini
--- a/Engine/Config/Windows/DataDrivenPlatformInfo.ini
+++ b/Engine/Config/Windows/DataDrivenPlatformInfo.ini
@@ -85,7 +85,8 @@ MaxMeshShaderThreadGroupSize=128
 bSupportsDebugViewShaders=true
 bSupportsRayTracing=true
 bSupportsRayTracingShaders=true
-bSupportsInlineRayTracing=false
+bSupportsInlineRayTracing=true
+bInlineRayTracingRequiresBindless=true
 bSupportsRayTracingCallableShaders=true
 bSupportsRayTracingProceduralPrimitive=true
 bSupportsRayTracingIndirectInstanceData=true
diff --git a/Engine/Shaders/Private/RayTracing/RayTracingCommon.ush b/Engine/Shaders/Private/RayTracing/RayTracingCommon.ush
--- a/Engine/Shaders/Private/RayTracing/RayTracingCommon.ush
+++ b/Engine/Shaders/Private/RayTracing/RayTracingCommon.ush
@@ -13,7 +13,7 @@
 #include "/Engine/Public/ShaderVersion.ush"
 
 // Update this GUID to invalidate RayTracing shaders and force a shader recompilation.
-#pragma message("UESHADERMETADATA_VERSION A2C2DC79-420D-4512-B864-63BB2951C211")
+#pragma message("UESHADERMETADATA_VERSION 82EF3461-600F-45F4-94A7-F2D9EF032386")
 
 #include "../ShadingCommon.ush"
 #include "../OctahedralCommon.ush"
diff --git a/Engine/Shaders/Private/RayTracing/TraceRayInlineCommon.ush b/Engine/Shaders/Private/RayTracing/TraceRayInlineCommon.ush
--- a/Engine/Shaders/Private/RayTracing/TraceRayInlineCommon.ush
+++ b/Engine/Shaders/Private/RayTracing/TraceRayInlineCommon.ush
@@ -3,7 +3,7 @@
 #pragma once
 
 // Update this GUID to invalidate inline RayTracing shaders and force a shader recompilation.
-#pragma message("UESHADERMETADATA_VERSION 1252ECC0-45EC-4D0F-8A87-393D1DDA0C75")
+#pragma message("UESHADERMETADATA_VERSION 06A6E5F1-49BF-4FF3-9D90-4B501E7E5FC4")
 
 struct FTraceRayInlineTraversalStatistics
 {
diff --git a/Engine/Source/Runtime/D3D12RHI/Private/D3D12Adapter.cpp b/Engine/Source/Runtime/D3D12RHI/Private/D3D12Adapter.cpp
--- a/Engine/Source/Runtime/D3D12RHI/Private/D3D12Adapter.cpp
+++ b/Engine/Source/Runtime/D3D12RHI/Private/D3D12Adapter.cpp
@@ -1043,8 +1043,13 @@ void FD3D12Adapter::InitializeDevices()
 							GRHISupportsRayTracing = RHISupportsRayTracing(GMaxRHIShaderPlatform);
 							GRHISupportsRayTracingShaders = GRHISupportsRayTracing && RHISupportsRayTracingShaders(GMaxRHIShaderPlatform);
 
+							ERHIBindlessConfiguration BindlessResourcesConfiguration = RHIGetRuntimeBindlessResourcesConfiguration(GMaxRHIShaderPlatform);
+
 							GRHISupportsRayTracingPSOAdditions = true;
-							GRHISupportsInlineRayTracing = GRHISupportsRayTracing && RHISupportsInlineRayTracing(GMaxRHIShaderPlatform) && (GMaxRHIFeatureLevel >= ERHIFeatureLevel::SM6);
+							GRHISupportsInlineRayTracing = GRHISupportsRayTracing && RHISupportsInlineRayTracing(GMaxRHIShaderPlatform) && (GMaxRHIFeatureLevel >= ERHIFeatureLevel::SM6) && (BindlessResourcesConfiguration != ERHIBindlessConfiguration::Disabled);
+														
+							// Inline RayTracing SBT is required
+							GRHIGlobals.RayTracing.RequiresInlineRayTracingSBT = GRHISupportsInlineRayTracing;
 						}
 						else
 						{
diff --git a/Engine/Source/Runtime/D3D12RHI/Private/D3D12Commands.cpp b/Engine/Source/Runtime/D3D12RHI/Private/D3D12Commands.cpp
--- a/Engine/Source/Runtime/D3D12RHI/Private/D3D12Commands.cpp
+++ b/Engine/Source/Runtime/D3D12RHI/Private/D3D12Commands.cpp
@@ -135,6 +135,8 @@ FD3D12ResourceLocation& FD3D12CommandContext::SetupIndirectArgument(FRHIBuffer*
 
 void FD3D12CommandContext::PostGpuEvent()
 {
+	UnsetExplicitDescriptorCache();
+
 	ConditionalSplitCommandList();
 	DEBUG_EXECUTE_COMMAND_LIST(this);
 }
diff --git a/Engine/Source/Runtime/D3D12RHI/Private/D3D12StateCache.cpp b/Engine/Source/Runtime/D3D12RHI/Private/D3D12StateCache.cpp
--- a/Engine/Source/Runtime/D3D12RHI/Private/D3D12StateCache.cpp
+++ b/Engine/Source/Runtime/D3D12RHI/Private/D3D12StateCache.cpp
@@ -4,6 +4,7 @@
 //	thread performance by removing redundant device context calls.
 
 #include "D3D12RHIPrivate.h"
+#include "D3D12ExplicitDescriptorCache.h"
 
 // This value defines how many descriptors will be in the device global descriptor heap. This heap contains all shader visible view descriptors.
 // Other shader visible descriptor heaps (e.g. OnlineViewHeap) are allocated from this pool. Non-visible heaps (e.g. LocalViewHeap) are allocated as standalone.
@@ -456,11 +457,6 @@ void FD3D12StateCache::ApplyState(ERHIPipeline HardwarePipe, ED3D12PipelineType
 		return;
 	}
 
-	const bool bRootSignatureChanged = InternalSetRootSignature(PipelineType, PSOCommonData->RootSignature);
-
-	// Ensure the correct graphics PSO is set.
-	InternalSetPipelineState(PSOCommonData->PipelineState);
-
 	bool bBindlessResources = PSOCommonData->RootSignature->UsesDynamicResources();
 	bool bBindlessSamplers = PSOCommonData->RootSignature->UsesDynamicSamplers();
 
@@ -477,14 +473,34 @@ void FD3D12StateCache::ApplyState(ERHIPipeline HardwarePipe, ED3D12PipelineType
 		bool bHaveResourceHeap = BindlessManager.AreResourcesBindless(ERHIBindlessConfiguration::AllShaders);
 		bool bHaveSamplerHeap = BindlessManager.AreSamplersBindless(ERHIBindlessConfiguration::AllShaders);
 
-		checkf(!bBindlessResources || bHaveResourceHeap, TEXT("Using dynamic samplers without the bindless sampler heap configured. Please check your configuration."));
-		checkf(!bBindlessSamplers  || bHaveSamplerHeap, TEXT("Using dynamic samplers without the bindless sampler heap configured. Please check your configuration."));
+		bool bRequireExplicitDescriptorCacheSet = bBindlessResources && !bHaveResourceHeap;
+		if (bRequireExplicitDescriptorCacheSet)
+		{
+			// Samplers should be bindless as well
+			check(bBindlessSamplers);
+			// Should have any SRVs or UAVs
+			check(!bHasTableResources);
+
+			// Everything bindless, don't require any CVS, SRV and UAV view data
+			FD3D12ExplicitDescriptorCache TransientDescriptorCache(GetParentDevice(), 1);
+			TransientDescriptorCache.Init(0, 0, 0, ERHIBindlessConfiguration::RayTracingShaders);
+
+			CmdContext.SetExplicitDescriptorCache(TransientDescriptorCache);
+			
+			PipelineState.Compute.bNeedSetRootSignature = true;
+			PipelineState.Graphics.bNeedSetRootSignature = true;
+		}
 
 		check(!(bHaveResourceHeap && bHasTableResources));
 		check(!(bHaveSamplerHeap  && bHasSamplers      ));
 	}
 #endif
 
+	const bool bRootSignatureChanged = InternalSetRootSignature(PipelineType, PSOCommonData->RootSignature);
+
+	// Ensure the correct graphics PSO is set.
+	InternalSetPipelineState(PSOCommonData->PipelineState);
+
 	if (bRootSignatureChanged)
 	{
 		CmdContext.BindDiagnosticBuffer(PSOCommonData->RootSignature, PipelineType);
diff --git a/Engine/Source/Runtime/Engine/Private/ShaderCompiler/ShaderCompiler.cpp b/Engine/Source/Runtime/Engine/Private/ShaderCompiler/ShaderCompiler.cpp
--- a/Engine/Source/Runtime/Engine/Private/ShaderCompiler/ShaderCompiler.cpp
+++ b/Engine/Source/Runtime/Engine/Private/ShaderCompiler/ShaderCompiler.cpp
@@ -3977,11 +3977,20 @@ void GlobalBeginCompileShader(
 	if (RHIGetBindlessSupport(ShaderPlatform) != ERHIBindlessSupport::Unsupported && !Input.Environment.CompilerFlags.Contains(CFLAG_ForceBindful))
 	{
 		const bool bIsRaytracingShader = IsRayTracingShaderFrequency(Input.Target.GetFrequency());
+		const bool bIsInlineRaytracingShader = Input.Environment.CompilerFlags.Contains(CFLAG_InlineRayTracing);
+		const bool bInlineRayTracingRequiresBindless = FDataDrivenShaderPlatformInfo::GetRequiresBindlessForInlineRayTracing(EShaderPlatform(Target.Platform));
 
 		const ERHIBindlessConfiguration ResourcesConfig = UE::ShaderCompiler::GetBindlessResourcesConfiguration(ShaderFormatName);
 		const ERHIBindlessConfiguration SamplersConfig = UE::ShaderCompiler::GetBindlessSamplersConfiguration(ShaderFormatName);
 
-		if (ResourcesConfig == ERHIBindlessConfiguration::AllShaders || (ResourcesConfig == ERHIBindlessConfiguration::RayTracingShaders && bIsRaytracingShader))
+		// If inline raytracing is enabled and bindless is required make sure bindless is enabled for raytracing at least
+		if (bIsInlineRaytracingShader && bInlineRayTracingRequiresBindless)
+		{
+			check(ResourcesConfig != ERHIBindlessConfiguration::Disabled);
+		}
+
+		bool EnableBindlessForInlineRaytracing = (bIsInlineRaytracingShader && bInlineRayTracingRequiresBindless && ResourcesConfig != ERHIBindlessConfiguration::Disabled);
+		if (ResourcesConfig == ERHIBindlessConfiguration::AllShaders || (ResourcesConfig == ERHIBindlessConfiguration::RayTracingShaders && bIsRaytracingShader) || EnableBindlessForInlineRaytracing)
 		{
 			bEnableBindlessMacro = true;
 			Input.Environment.CompilerFlags.Add(CFLAG_BindlessResources);
diff --git a/Engine/Source/Runtime/RHI/Private/DataDrivenShaderPlatformInfo.cpp b/Engine/Source/Runtime/RHI/Private/DataDrivenShaderPlatformInfo.cpp
--- a/Engine/Source/Runtime/RHI/Private/DataDrivenShaderPlatformInfo.cpp
+++ b/Engine/Source/Runtime/RHI/Private/DataDrivenShaderPlatformInfo.cpp
@@ -266,6 +266,7 @@ void FGenericDataDrivenShaderPlatformInfo::ParseDataDrivenShaderInfo(const FConf
 	GET_SECTION_BOOL_HELPER(bSupportsRayTracing);
 	GET_SECTION_BOOL_HELPER(bSupportsRayTracingShaders);
 	GET_SECTION_BOOL_HELPER(bSupportsInlineRayTracing);
+	GET_SECTION_BOOL_HELPER(bInlineRayTracingRequiresBindless);
 	GET_SECTION_BOOL_HELPER(bSupportsRayTracingCallableShaders);
 	GET_SECTION_BOOL_HELPER(bSupportsRayTracingProceduralPrimitive);
 	GET_SECTION_BOOL_HELPER(bSupportsRayTracingTraversalStatistics);
@@ -571,6 +572,7 @@ void FGenericDataDrivenShaderPlatformInfo::UpdatePreviewPlatforms()
 				PREVIEW_USE_RUNTIME_VALUE(bSupportsUniformBufferObjects);
 
 				PREVIEW_USE_RUNTIME_VALUE(bSupportsInlineRayTracing);
+				PREVIEW_USE_RUNTIME_VALUE(bInlineRayTracingRequiresBindless);
 				PREVIEW_USE_RUNTIME_VALUE(bSupportsRayTracingShaders);
 
 				// Settings that will never be supported in preview
diff --git a/Engine/Source/Runtime/RHI/Public/DataDrivenShaderPlatformInfo.h b/Engine/Source/Runtime/RHI/Public/DataDrivenShaderPlatformInfo.h
--- a/Engine/Source/Runtime/RHI/Public/DataDrivenShaderPlatformInfo.h
+++ b/Engine/Source/Runtime/RHI/Public/DataDrivenShaderPlatformInfo.h
@@ -105,6 +105,7 @@ class FGenericDataDrivenShaderPlatformInfo
 	uint32 bOverrideFMaterial_NeedsGBufferEnabled : 1;
 	uint32 bSupportsFFTBloom : 1;
 	uint32 bSupportsInlineRayTracing : 1;
+	uint32 bInlineRayTracingRequiresBindless : 1;
 	uint32 bSupportsRayTracingShaders : 1;
 	uint32 bSupportsVertexShaderLayer : 1;
 	uint32 BindlessSupport : int32(ERHIBindlessSupport::NumBits);
@@ -376,6 +377,12 @@ public:
 		return Infos[Platform].bSupportsRayTracing && Infos[Platform].bSupportsInlineRayTracing;
 	}
 
+	static FORCEINLINE_DEBUGGABLE const bool GetRequiresBindlessForInlineRayTracing(const FStaticShaderPlatform Platform)
+	{
+		check(IsValid(Platform));
+		return Infos[Platform].bInlineRayTracingRequiresBindless;
+	}
+
 	static FORCEINLINE_DEBUGGABLE const bool GetSupportsRayTracingCallableShaders(const FStaticShaderPlatform Platform)
 	{
 		check(IsValid(Platform));
