- Don't make partial RTPSO compilations fatal when they don't compile
- Disable workaround to not compile RTPSO from Shader bundles by default

#okforgiyhub public

diff --git a/Engine/Source/Runtime/D3D12RHI/Private/D3D12RayTracing.cpp b/Engine/Source/Runtime/D3D12RHI/Private/D3D12RayTracing.cpp
--- a/Engine/Source/Runtime/D3D12RHI/Private/D3D12RayTracing.cpp
+++ b/Engine/Source/Runtime/D3D12RHI/Private/D3D12RayTracing.cpp
@@ -652,7 +652,19 @@ static TRefCountPtr<ID3D12StateObject> CreateRayTracingStateObject(
 	}
 #endif
 
-	VERIFYD3D12RESULT(RayTracingDevice->CreateStateObject(&Desc, IID_PPV_ARGS(Result.GetInitReference())));
+	HRESULT hr = RayTracingDevice->CreateStateObject(&Desc, IID_PPV_ARGS(Result.GetInitReference()));
+	if (FAILED(hr))
+	{
+		// Failed state objects are not fatal if it's coming from partial/non-required RTPSO 
+		// Will be fatal when it's actually needed for rendering
+		UE_LOG(LogD3D12RHI, Warning, TEXT("Failed to create raytracing RTPSO - add '-d3ddebug' to the command line to get more info on RTPSO compilation errors from the debug device"));
+
+		// Catch GPU crashes explicitly.
+		if (hr == DXGI_ERROR_DEVICE_REMOVED || hr == DXGI_ERROR_DEVICE_HUNG)
+		{
+			VERIFYD3D12RESULT_EX(hr, RayTracingDevice);
+		}
+	}
 
 #if WITH_NVAPI
 	if (bReorderingDisabled)
@@ -934,6 +946,10 @@ public:
 		{
 			check(bDeserialized || (CompileEvent.IsValid() && CompileEvent->IsComplete()));
 			check(StateObject);
+			if (StateObject == nullptr)
+			{
+				UE_LOG(LogD3D12RHI, Fatal, TEXT("Trying to add a raytracing state collection which failed to compile to a RTPSO collection."));
+			}
 
 			D3D12_EXISTING_COLLECTION_DESC Result = {};
 			Result.pExistingCollection = StateObject;
@@ -994,12 +1010,14 @@ public:
 				FEntry& InEntry,
 				FKey InCacheKey,
 				FD3D12Device* InDevice,
-				ECollectionType InCollectionType)
+				ECollectionType InCollectionType,
+				bool bInRequired)
 			: Entry(InEntry)
 			, CacheKey(InCacheKey)
 			, Device(InDevice)
 			, RayTracingDevice(InDevice->GetDevice5())
 			, CollectionType(InCollectionType)
+			, bRequired(bInRequired)
 		{
 		}
 
@@ -1096,10 +1114,14 @@ public:
 			if (Entry.StateObject)
 			{
 				Device->GetRayTracingPipelineInfo(Entry.StateObject, &Entry.PipelineInfo);
-			}
 
-			// Retrieve the identifier from the library
-			Entry.Identifier = GetShaderIdentifier(Entry.StateObject, Entry.GetPrimaryExportNameChars());
+				// Retrieve the identifier from the library
+				Entry.Identifier = GetShaderIdentifier(Entry.StateObject, Entry.GetPrimaryExportNameChars());
+			}
+			else if (bRequired)
+			{
+				UE_LOG(LogD3D12RHI, Fatal, TEXT("Failed to create a required RTPSO pipeline state collection"));
+			}
 			
 			CompileTimeCycles += FPlatformTime::Cycles64();
 
@@ -1127,7 +1149,7 @@ public:
 		FD3D12Device* Device;
 		ID3D12Device5* RayTracingDevice;
 		ECollectionType CollectionType;
-
+		bool bRequired;
 	};
 
 	FEntry* GetOrCompileShader(
@@ -1136,6 +1158,7 @@ public:
 		ID3D12RootSignature* GlobalRootSignature,
 		uint32 MaxAttributeSizeInBytes,
 		uint32 MaxPayloadSizeInBytes,
+		bool bRequired,
 		ECollectionType CollectionType,
 		FGraphEventArray& CompletionList,
 		bool* bOutCacheHit = nullptr)
@@ -1210,7 +1233,8 @@ public:
 					Entry,
 					CacheKey,
 					Device,
-					CollectionType
+					CollectionType,
+					bRequired
 				);
 			}
 		}
@@ -1224,7 +1248,12 @@ public:
 		}
 		else
 		{
-			check(FindResult->StateObject != nullptr);
+			// required PSOs collection are fatal if not valid (can have failed during precaching of this
+			// RTPSO collection and wasn't fatal during precaching, but did fill the cache)
+			if (bRequired && FindResult->StateObject == nullptr)
+			{
+				UE_LOG(LogD3D12RHI, Fatal, TEXT("Failed to create a required RTPSO collection"));
+			}
 		}
 
 		return FindResult;
@@ -2761,7 +2790,12 @@ FD3D12RayTracingPipelineState::FD3D12RayTracingPipelineState(FD3D12Device* Devic
 	{
 		// verify that that the same shader binding layout is used for all shaders in the RTPSO or not sampling any resources
 		uint32 TotalResourceCount = Shader->ResourceCounts.NumCBs + Shader->ResourceCounts.NumSRVs + Shader->ResourceCounts.NumUAVs + Shader->ResourceCounts.NumSamplers;
-		checkf(TotalResourceCount == 0 || Shader->ShaderBindingLayoutHash == ShaderBindingLayoutHash, TEXT("Raytracing shader with with entry point %s and name %s doesn't match the RTPSO ShaderBindingLayout"), *Shader->EntryPoint, Shader->GetShaderName());
+		if (TotalResourceCount != 0 && Shader->ShaderBindingLayoutHash != ShaderBindingLayoutHash)
+		{
+			UE_LOG(LogD3D12RHI, Warning, TEXT("Raytracing shader with with entry point %s, name %s and ShaderBindingLayout hash %u doesn't match the RTPSO ShaderBindingLayout hash %u."),
+			       *Shader->EntryPoint, Shader->GetShaderName(), Shader->ShaderBindingLayoutHash, ShaderBindingLayoutHash);
+			ensure(false);
+		}
 
 		bool bIsAlreadyInSet = false;
 		const uint64 ShaderHash = GetShaderHash64(Shader);
@@ -2775,6 +2809,7 @@ FD3D12RayTracingPipelineState::FD3D12RayTracingPipelineState(FD3D12Device* Devic
 			Device, Shader, GlobalRootSignature,
 			Initializer.MaxAttributeSizeInBytes,
 			Initializer.MaxPayloadSizeInBytes,
+			!Initializer.bPartial, //< partial PSOs are not required
 			CollectionType, CompileCompletionList,
 			&bCacheHit);
 
@@ -2968,6 +3003,11 @@ FD3D12RayTracingPipelineState::FD3D12RayTracingPipelineState(FD3D12Device* Devic
 			{}, // LocalRootSignatureAssociations,
 			UniqueShaderCollectionDescs,
 			D3D12_STATE_OBJECT_TYPE_RAYTRACING_PIPELINE);
+
+		if (StateObject == nullptr)
+		{
+			UE_LOG(LogD3D12RHI, Fatal, TEXT("Failed to a create raytracing pipeline state"));
+		}
 	}
 
 	if (GRayTracingSpecializeStateObjects != 0 && Initializer.GetRayGenTable().Num() > 1)
diff --git a/Engine/Source/Runtime/RenderCore/Private/ShaderPipelineCache.cpp b/Engine/Source/Runtime/RenderCore/Private/ShaderPipelineCache.cpp
--- a/Engine/Source/Runtime/RenderCore/Private/ShaderPipelineCache.cpp
+++ b/Engine/Source/Runtime/RenderCore/Private/ShaderPipelineCache.cpp
@@ -271,7 +271,7 @@ static FAutoConsoleVariableRef CVarShaderPipelineCacheDoNotPrecompileComputePSO(
 												);
 
 // Ray tracing PSOs are recorded into the cache incorrectly due to a bug.
-bool GShaderPipelineCacheDoNotPrecompileRayTracingPSO = true;
+bool GShaderPipelineCacheDoNotPrecompileRayTracingPSO = false;
 static FAutoConsoleVariableRef CVarShaderPipelineCacheDoNotPrecompileRayTracingPSO(
 	TEXT("r.ShaderPipelineCache.DoNotPrecompileRayTracingPSO"),
 	GShaderPipelineCacheDoNotPrecompileRayTracingPSO,
