Re-enable D3D12 dummy backbuffer to work around crashes due to exiting or re-entering exclusive fullscreen where the backbuffer index gets out of sync.

This is not a proper fix but it significantly decreases the chance of crashing.

Also use GetCurrentBackbufferIndex to retrieve the backbuffer index on swapchain (re-)creation instead of assuming 0.

#rnx

diff --git a/Engine/Source/Runtime/D3D12RHI/Private/Windows/WindowsD3D12RHIDefinitions.h b/Engine/Source/Runtime/D3D12RHI/Private/Windows/WindowsD3D12RHIDefinitions.h
--- a/Engine/Source/Runtime/D3D12RHI/Private/Windows/WindowsD3D12RHIDefinitions.h
+++ b/Engine/Source/Runtime/D3D12RHI/Private/Windows/WindowsD3D12RHIDefinitions.h
@@ -26,7 +26,7 @@ typedef ID3D12GraphicsCommandList ID3D12CopyCommandList;
 
 #define USE_STATIC_ROOT_SIGNATURE			0
 #define D3D12RHI_USE_CONSTANT_BUFFER_VIEWS	0
-#define D3D12RHI_USE_DUMMY_BACKBUFFER		0
+#define D3D12RHI_USE_DUMMY_BACKBUFFER		1
 #define D3D12RHI_USE_SDR_BACKBUFFER			0
 #define D3D12RHI_SUPPORTS_UAV_BACKBUFFER	0
 #define D3D12RHI_USE_DXGI_COLOR_SPACE		1
diff --git a/Engine/Source/Runtime/D3D12RHI/Private/Windows/WindowsD3D12Viewport.cpp b/Engine/Source/Runtime/D3D12RHI/Private/Windows/WindowsD3D12Viewport.cpp
--- a/Engine/Source/Runtime/D3D12RHI/Private/Windows/WindowsD3D12Viewport.cpp
+++ b/Engine/Source/Runtime/D3D12RHI/Private/Windows/WindowsD3D12Viewport.cpp
@@ -318,13 +318,22 @@ void FD3D12Viewport::ResizeInternal()
 		}
 	}
 
-	SetBackBufferIndex_RHIThread(0);
+	uint32 CurrentBackBufferIndex = 0;
+
+#if DXGI_MAX_SWAPCHAIN_INTERFACE >= 3
+	if (SwapChain3)
+	{
+		CurrentBackBufferIndex = SwapChain3->GetCurrentBackBufferIndex();
+	}
+#endif
+
+	SetBackBufferIndex_RHIThread(CurrentBackBufferIndex);
 
 #if D3D12RHI_USE_DUMMY_BACKBUFFER
 	// Create dummy back buffer which always reference to the actual RHI thread back buffer - can't be bound directly to D3D12
 	DummyBackBuffer_RenderThread = CreateDummyBackBufferTextures(Adapter, PixelFormat, SizeX, SizeY);
 #else
-	SetBackBufferIndex_RenderThread(0);
+	SetBackBufferIndex_RenderThread(CurrentBackBufferIndex);
 #endif
 }
 
