Prevent crash when using the D3D12 dummy backbuffer.

When enabled, the dummy backbuffer will be the only one marked as presentable, otherwise the real backbuffer will try to perform an invalid cast to the dummy backbuffer structure.

#rnx

diff --git a/Engine/Source/Runtime/D3D12RHI/Private/D3D12Viewport.cpp b/Engine/Source/Runtime/D3D12RHI/Private/D3D12Viewport.cpp
--- a/Engine/Source/Runtime/D3D12RHI/Private/D3D12Viewport.cpp
+++ b/Engine/Source/Runtime/D3D12RHI/Private/D3D12Viewport.cpp
@@ -87,7 +87,12 @@ FD3D12Texture* GetSwapChainSurface(FD3D12Device* Parent, EPixelFormat PixelForma
 
 	FString Name = FString::Printf(TEXT("BackBuffer%d"), BackBufferIndex);
 	
+#if D3D12RHI_USE_DUMMY_BACKBUFFER
+	ETextureCreateFlags SwapchainTextureCreateFlags = ETextureCreateFlags::RenderTargetable | ETextureCreateFlags::ResolveTargetable;
+#else
 	ETextureCreateFlags SwapchainTextureCreateFlags = ETextureCreateFlags::RenderTargetable | ETextureCreateFlags::Presentable | ETextureCreateFlags::ResolveTargetable;
+#endif
+
 #if D3D12RHI_SUPPORTS_UAV_BACKBUFFER
 	SwapchainTextureCreateFlags |= ETextureCreateFlags::UAV;
 #endif
@@ -133,7 +138,11 @@ FD3D12Texture* GetSwapChainSurface(FD3D12Device* Parent, EPixelFormat PixelForma
 				&NewTexture->ResourceLocation,
 				NewTexture,
 				PixelFormat,
-				TexCreate_RenderTargetable |  TexCreate_ShaderResource | TexCreate_Presentable,
+#if D3D12RHI_USE_DUMMY_BACKBUFFER
+				TexCreate_RenderTargetable | TexCreate_ShaderResource,
+#else
+				TexCreate_RenderTargetable | TexCreate_ShaderResource | TexCreate_Presentable,
+#endif
 				InitialState,
 				*Name);
 		}
