Replace check with ensure for SBT persistent state validation - shouldn't be fatal check for users

diff --git a/Engine/Source/Runtime/D3D12RHI/Private/D3D12RayTracing.cpp b/Engine/Source/Runtime/D3D12RHI/Private/D3D12RayTracing.cpp
--- a/Engine/Source/Runtime/D3D12RHI/Private/D3D12RayTracing.cpp
+++ b/Engine/Source/Runtime/D3D12RHI/Private/D3D12RayTracing.cpp
@@ -1314,7 +1314,7 @@ private:
 	void CompareData(uint32 Offset, const void* InData, uint32 InDataSize)
 	{
 		const void* CurrentData = Data.GetData() + Offset;
-		check(FMemory::Memcmp(CurrentData, InData, InDataSize) == 0);
+		ensure(FMemory::Memcmp(CurrentData, InData, InDataSize) == 0);
 	}
 
 	void WriteLocalShaderRecord(uint32 ShaderTableOffset, uint32 RecordIndex, uint32 OffsetWithinRecord, const void* InData, uint32 InDataSize)
