Also track residency for resources referenced by the TLAS in bindless mode

diff --git a/Engine/Source/Runtime/D3D12RHI/Private/D3D12DescriptorCache.cpp b/Engine/Source/Runtime/D3D12RHI/Private/D3D12DescriptorCache.cpp
--- a/Engine/Source/Runtime/D3D12RHI/Private/D3D12DescriptorCache.cpp
+++ b/Engine/Source/Runtime/D3D12RHI/Private/D3D12DescriptorCache.cpp
@@ -516,6 +516,14 @@ void FD3D12DescriptorCache::PrepareBindlessViews(EShaderFrequency ShaderStage, T
 		if (ensure(SRV))
 		{
 			Context.UpdateResidency(SRV->GetResource());
+
+#if RHI_RAYTRACING
+			FD3D12RayTracingScene* RayTracingScene = SRV->GetRayTracingScene();
+			if (RayTracingScene)
+			{
+				RayTracingScene->UpdateResidency(Context);
+			}
+#endif
 		}
 	}
 
diff --git a/Engine/Source/Runtime/D3D12RHI/Private/D3D12RayTracing.cpp b/Engine/Source/Runtime/D3D12RHI/Private/D3D12RayTracing.cpp
--- a/Engine/Source/Runtime/D3D12RHI/Private/D3D12RayTracing.cpp
+++ b/Engine/Source/Runtime/D3D12RHI/Private/D3D12RayTracing.cpp
@@ -1783,6 +1783,12 @@ public:
 				D3D12_RESOURCE_STATE_COPY_DEST,
 				D3D12_RESOURCE_STATE_NON_PIXEL_SHADER_RESOURCE,
 				0);
+
+#if ENABLE_RESIDENCY_MANAGEMENT
+			// assume no references resources when inline SBT only - inline SBT only has references to IB/VB buffers and these are 
+			// residency tracked via the TLAS
+			check(ShaderBindingMode != ERayTracingShaderBindingMode::Inline || ReferencedResources.IsEmpty());
+#endif // ENABLE_RESIDENCY_MANAGEMENT
 		}
 
 		bIsDirty = false;
