Use D3D12 message callbacks instead of the exceptions hook for logging D3D Debug message so we can properly reroute Warning and Info message as well into UE_LOG
Fixed "break on warning" logic
https://devblogs.microsoft.com/directx/d3d12-debug-layer-message-callback/

diff --git a/Engine/Source/Runtime/D3D12RHI/Private/D3D12Adapter.cpp b/Engine/Source/Runtime/D3D12RHI/Private/D3D12Adapter.cpp
--- a/Engine/Source/Runtime/D3D12RHI/Private/D3D12Adapter.cpp
+++ b/Engine/Source/Runtime/D3D12RHI/Private/D3D12Adapter.cpp
@@ -75,7 +75,7 @@ bool D3D12_ShouldBreakOnD3DDebugErrors()
 
 bool D3D12_ShouldBreakOnD3DDebugWarnings()
 {
-	return GD3D12DebugCvar.GetValueOnAnyThread() > 3;
+	return GD3D12DebugCvar.GetValueOnAnyThread() == 3;
 }
 
 #if PLATFORM_WINDOWS
@@ -138,6 +138,41 @@ static FAutoConsoleVariableRef CVarD3D12EvictAllResidentResourcesInBackground(
 #endif // PLATFORM_WINDOWS
 
 #if D3D12_SUPPORTS_INFO_QUEUE
+
+// Log with the UE_LOG the message received from D3D12 debug layer
+// Return if we need to break execution or not
+static bool LogD3D12Message(D3D12_MESSAGE_SEVERITY Severity, LPCSTR pDescription)
+{
+	bool bBreakOnLog = false;
+
+	if (pDescription)
+	{
+		switch (Severity)
+		{
+		case D3D12_MESSAGE_SEVERITY_CORRUPTION:
+		case D3D12_MESSAGE_SEVERITY_ERROR:
+		{
+			UE_LOG(LogD3D12RHI, Error, TEXT("[D3DDebug] %s"), ANSI_TO_TCHAR(pDescription));
+			bBreakOnLog = D3D12_ShouldBreakOnD3DDebugErrors();
+			break;
+		}
+		case D3D12_MESSAGE_SEVERITY_WARNING:
+		{
+			UE_LOG(LogD3D12RHI, Warning, TEXT("[D3DDebug] %s"), ANSI_TO_TCHAR(pDescription));
+			bBreakOnLog = D3D12_ShouldBreakOnD3DDebugWarnings();
+			break;
+		}
+		default:
+		{
+			UE_LOG(LogD3D12RHI, Log, TEXT("[D3DDebug] %s"), ANSI_TO_TCHAR(pDescription));
+			break;
+		}
+		}
+	}
+
+	return bBreakOnLog;
+}
+
 static bool CheckD3DStoredMessages()
 {
 	bool bResult = false;
@@ -152,8 +187,6 @@ static bool CheckD3DStoredMessages()
 			D3D12_MESSAGE* d3dMessage = nullptr;
 			SIZE_T AllocateSize = 0;
 
-			static const bool bBreakOnWarning = FParse::Param(FCommandLine::Get(), TEXT("d3dbreakonwarning"));
-
 			int StoredMessageCount = d3dInfoQueue->GetNumStoredMessagesAllowedByRetrievalFilter();
 			for (int MessageIndex = 0; MessageIndex < StoredMessageCount; MessageIndex++)
 			{
@@ -186,29 +219,7 @@ static bool CheckD3DStoredMessages()
 					// get the actual message data from the queue
 					hr = d3dInfoQueue->GetMessage(MessageIndex, d3dMessage, &MessageLength);
 
-
-					switch (d3dMessage->Severity)
-					{
-					case D3D12_MESSAGE_SEVERITY_CORRUPTION:
-					case D3D12_MESSAGE_SEVERITY_ERROR:
-						{
-							UE_LOG(LogD3D12RHI, Error, TEXT("[D3DDebug] %s"), ANSI_TO_TCHAR(d3dMessage->pDescription));
-							bResult = true;
-							break;
-						}
-					case D3D12_MESSAGE_SEVERITY_WARNING:
-						{
-							UE_LOG(LogD3D12RHI, Warning, TEXT("[D3DDebug] %s"), ANSI_TO_TCHAR(d3dMessage->pDescription));
-							if (bBreakOnWarning)
-								bResult = true;
-							break;
-						}
-					default:
-						{
-							UE_LOG(LogD3D12RHI, Log, TEXT("[D3DDebug] %s"), ANSI_TO_TCHAR(d3dMessage->pDescription));
-							break;
-						}
-					}
+					bResult = LogD3D12Message(d3dMessage->Severity, d3dMessage->pDescription);
 				}
 			}
 			d3dInfoQueue->ClearStoredMessages();
@@ -223,21 +234,26 @@ static bool CheckD3DStoredMessages()
 }
 #endif // #if D3D12_SUPPORTS_INFO_QUEUE
 
-#if PLATFORM_WINDOWS
+DWORD GD3D12CallBackCookie = 0;
 
+#if PLATFORM_WINDOWS
 /** Handle d3d messages and write them to the log file **/
 static LONG __stdcall D3DVectoredExceptionHandler(EXCEPTION_POINTERS* InInfo)
 {
 	// Only handle D3D error codes here
 	if (InInfo->ExceptionRecord->ExceptionCode == _FACDXGI)
 	{
-		if (CheckD3DStoredMessages())
+		// Logging and break in code is managed in the D3D12MessageCallBack if ID3D12InfoQueue1 but we handle and continue D3D exception in order to not break while the debugger is not attached
+		if (GD3D12CallBackCookie == 0)
 		{
-			if (FPlatformMisc::IsDebuggerPresent())
+			if (CheckD3DStoredMessages())
 			{
-				// when we get here, then it means that BreakOnSeverity was set for this error message, so request the debug break here as well
-				// when the debugger is attached
-				UE_DEBUG_BREAK();
+				if (FPlatformMisc::IsDebuggerPresent())
+				{
+					// when we get here, then it means that BreakOnSeverity was set for this error message, so request the debug break here as well
+					// when the debugger is attached
+					UE_DEBUG_BREAK();
+				}
 			}
 		}
 
@@ -248,9 +264,22 @@ static LONG __stdcall D3DVectoredExceptionHandler(EXCEPTION_POINTERS* InInfo)
 	// continue searching
 	return EXCEPTION_CONTINUE_SEARCH;
 }
-
 #endif // #if PLATFORM_WINDOWS
 
+#if D3D12_SUPPORTS_INFO_QUEUE
+void D3D12MessageCallBack(D3D12_MESSAGE_CATEGORY Category, D3D12_MESSAGE_SEVERITY Severity, D3D12_MESSAGE_ID ID, LPCSTR pDescription, void* pContext)
+{
+	bool bBreakOnLog = LogD3D12Message(Severity, pDescription);
+	
+	if (FPlatformMisc::IsDebuggerPresent() && bBreakOnLog)
+	{
+		// When we get here, then it means that BreakOnSeverity was set for this error message, so request the debug break here as well when the debugger is attached
+		UE_DEBUG_BREAK();
+	}
+
+}
+#endif
+
 
 FTransientUniformBufferAllocator::~FTransientUniformBufferAllocator()
 {
@@ -809,6 +838,24 @@ void FD3D12Adapter::CreateD3DInfoQueue()
 		{
 			pd3dInfoQueue->SetBreakOnSeverity(D3D12_MESSAGE_SEVERITY_WARNING, D3D12_ShouldBreakOnD3DDebugWarnings());
 		}
+		
+		// Try to use RegisterMessageCallback in the new ID3D12InfoQueue1 (available since Windows 10 Release Preview build 20236)
+		{
+			ID3D12InfoQueue1* pd3dInfoQueue1 = nullptr;
+			RootDevice->QueryInterface(__uuidof(ID3D12InfoQueue1), (void**)&pd3dInfoQueue1);
+			if (pd3dInfoQueue1)
+			{
+				void* ParamContext = this;
+				pd3dInfoQueue1->RegisterMessageCallback(D3D12MessageCallBack, D3D12_MESSAGE_CALLBACK_FLAG_NONE, ParamContext, &GD3D12CallBackCookie);
+
+				if (GD3D12CallBackCookie == 0)
+				{
+					UE_LOG(LogD3D12RHI, Log, TEXT("ID3D12InfoQueue1::RegisterMessageCallback failed."));
+				}
+
+				pd3dInfoQueue1->Release();
+			}
+		}
 
 		pd3dInfoQueue->Release();
 	}
@@ -1455,6 +1502,17 @@ FD3D12Adapter::~FD3D12Adapter()
 
 #if D3D12_SUPPORTS_INFO_QUEUE
 		CheckD3DStoredMessages();
+
+		if (GD3D12CallBackCookie != 0)
+		{
+			ID3D12InfoQueue1* pd3dInfoQueue1 = nullptr;
+			VERIFYD3D12RESULT(RootDevice->QueryInterface(__uuidof(ID3D12InfoQueue1), (void**)&pd3dInfoQueue1));
+			if (pd3dInfoQueue1)
+			{
+				pd3dInfoQueue1->UnregisterMessageCallback(GD3D12CallBackCookie);
+				GD3D12CallBackCookie = 0;
+			}
+		}
 #endif
 	}
 #endif
