Use GetTextureBaseRHI to get the real texture from the Dummy backbuffer texture instead of checking the TexCreate_Presentable flag

diff --git a/Engine/Source/Runtime/D3D12RHI/Private/D3D12RenderTarget.cpp b/Engine/Source/Runtime/D3D12RHI/Private/D3D12RenderTarget.cpp
--- a/Engine/Source/Runtime/D3D12RHI/Private/D3D12RenderTarget.cpp
+++ b/Engine/Source/Runtime/D3D12RHI/Private/D3D12RenderTarget.cpp
@@ -600,16 +600,6 @@ void FD3D12DynamicRHI::RHIReadSurfaceData(FRHITexture* InRHITexture, FIntRect In
 	const uint32 GPUIndex = InFlags.GetGPUIndex();
 	TArray<uint8> OutDataRaw;
 
-	// Could be back buffer reference texture, so get the correct D3D12 texture here
-	// We know already that it's a FD3D12Texture2D so cast is safe
-#if D3D12RHI_USE_DUMMY_BACKBUFFER
-	if (EnumHasAnyFlags(InRHITexture->GetFlags(), TexCreate_Presentable))
-	{
-		FD3D12BackBufferReferenceTexture2D* BufferBufferReferenceTexture = (FD3D12BackBufferReferenceTexture2D*)InRHITexture;
-		InRHITexture = BufferBufferReferenceTexture->GetBackBufferTexture();
-	}
-#endif
-
 	// Retrieve the base texture
 	FD3D12Texture* DestTexture2D = FD3D12CommandContext::RetrieveTexture(InRHITexture, GPUIndex);
 
diff --git a/Engine/Source/Runtime/D3D12RHI/Private/D3D12Texture.h b/Engine/Source/Runtime/D3D12RHI/Private/D3D12Texture.h
--- a/Engine/Source/Runtime/D3D12RHI/Private/D3D12Texture.h
+++ b/Engine/Source/Runtime/D3D12RHI/Private/D3D12Texture.h
@@ -53,7 +53,11 @@ public:
 #endif
 
 	// FRHITexture overrides
-	virtual void* GetTextureBaseRHI() override final			{ return this; }
+#if D3D12RHI_USE_DUMMY_BACKBUFFER
+	virtual void* GetTextureBaseRHI() override	{ return this; }
+#else
+	virtual void* GetTextureBaseRHI() override final { return this; }
+#endif
 	virtual void* GetNativeResource() const override final;
 	virtual FRHIDescriptorHandle GetDefaultBindlessHandle() const override;
 			
@@ -207,6 +211,9 @@ public:
 	FRHITexture* GetBackBufferTexture() const;
 	virtual FRHIDescriptorHandle GetDefaultBindlessHandle() const override;
 
+	// FRHITexture overrides
+	virtual void* GetTextureBaseRHI() override { return (void*)GetBackBufferTexture(); }
+
 private:
 	FD3D12Viewport* const Viewport;
 };
@@ -220,16 +227,7 @@ FORCEINLINE FD3D12Texture* GetD3D12TextureFromRHITexture(FRHITexture* Texture)
 		return NULL;
 	}
 	
-	// If it's the dummy backbuffer then swap with actual current RHI backbuffer right now
 	FRHITexture* RHITexture = Texture;
-#if D3D12RHI_USE_DUMMY_BACKBUFFER
-	if (RHITexture && EnumHasAnyFlags(RHITexture->GetFlags(), TexCreate_Presentable))
-	{
-		FD3D12BackBufferReferenceTexture2D* BufferBufferReferenceTexture = (FD3D12BackBufferReferenceTexture2D*)RHITexture;
-		RHITexture = BufferBufferReferenceTexture->GetBackBufferTexture();
-	}
-#endif
-
 	FD3D12Texture* Result((FD3D12Texture*)RHITexture->GetTextureBaseRHI());
 	check(Result);
 	return Result;
diff --git a/Engine/Source/Runtime/D3D12RHI/Private/D3D12Viewport.cpp b/Engine/Source/Runtime/D3D12RHI/Private/D3D12Viewport.cpp
--- a/Engine/Source/Runtime/D3D12RHI/Private/D3D12Viewport.cpp
+++ b/Engine/Source/Runtime/D3D12RHI/Private/D3D12Viewport.cpp
@@ -87,11 +87,7 @@ FD3D12Texture* GetSwapChainSurface(FD3D12Device* Parent, EPixelFormat PixelForma
 
 	FString Name = FString::Printf(TEXT("BackBuffer%d"), BackBufferIndex);
 	
-#if D3D12RHI_USE_DUMMY_BACKBUFFER
-	ETextureCreateFlags SwapchainTextureCreateFlags = ETextureCreateFlags::RenderTargetable | ETextureCreateFlags::ResolveTargetable;
-#else
 	ETextureCreateFlags SwapchainTextureCreateFlags = ETextureCreateFlags::RenderTargetable | ETextureCreateFlags::Presentable | ETextureCreateFlags::ResolveTargetable;
-#endif
 
 #if D3D12RHI_SUPPORTS_UAV_BACKBUFFER
 	SwapchainTextureCreateFlags |= ETextureCreateFlags::UAV;
@@ -138,11 +134,7 @@ FD3D12Texture* GetSwapChainSurface(FD3D12Device* Parent, EPixelFormat PixelForma
 				&NewTexture->ResourceLocation,
 				NewTexture,
 				PixelFormat,
-#if D3D12RHI_USE_DUMMY_BACKBUFFER
-				TexCreate_RenderTargetable | TexCreate_ShaderResource,
-#else
 				TexCreate_RenderTargetable | TexCreate_ShaderResource | TexCreate_Presentable,
-#endif
 				InitialState,
 				*Name);
 		}
