Fix errors reported by insights on non windows d3d12 platforms where double allocation events would occur: FD3D12DynamicRHI::LockBuffer -> FD3D12UploadHeapAllocator::AllocUploadResource -> FD3D12BuddyAllocator::Allocate would start tracing memory, and then call a new allocation event right before RenameLDAChain

Do not perform insights tracing in D3D12BufferStats::UpdateBufferStats since MemoryTrace_Alloc / MemoryTrace_Free already happen in the lower level allocator on unified memory platforms

diff --git a/Engine/Source/Runtime/D3D12RHI/Private/D3D12Allocation.cpp b/Engine/Source/Runtime/D3D12RHI/Private/D3D12Allocation.cpp
--- a/Engine/Source/Runtime/D3D12RHI/Private/D3D12Allocation.cpp
+++ b/Engine/Source/Runtime/D3D12RHI/Private/D3D12Allocation.cpp
@@ -484,9 +484,6 @@ void FD3D12BuddyAllocator::Allocate(uint32 SizeInBytes, uint32 Alignment, FD3D12
 	// which means that the memory would be counted twice. Because of this the tracking had to be disabled here.
 	// This does mean that non-buffer memory that goes through this allocator won't be tracked, so this does need a better solution.
 	// see UpdateBufferStats for a more detailed explanation.
-#if UE_MEMORY_TRACE_ENABLED
-	MemoryTrace_Alloc(ResourceLocation.GetGPUVirtualAddress(), SizeInBytes, Alignment, EMemoryTraceRootHeap::VideoMemory);
-#endif
 #endif
 }
 
@@ -546,9 +543,6 @@ void FD3D12BuddyAllocator::Deallocate(FD3D12ResourceLocation& ResourceLocation)
 		const uint32 AllocSize = uint32(OrderToUnitSize(Block.Data.Order) * MinBlockSize);
 		LLM_SCOPED_PAUSE_TRACKING_WITH_ENUM_AND_AMOUNT_BYTAG(D3D12AllocatorUnused, int64(AllocSize), ELLMTracker::Default, ELLMAllocType::System);
 	}
-#if UE_MEMORY_TRACE_ENABLED
-	MemoryTrace_Free(ResourceLocation.GetGPUVirtualAddress(), EMemoryTraceRootHeap::VideoMemory);
-#endif
 #endif
 }
 
diff --git a/Engine/Source/Runtime/D3D12RHI/Private/D3D12Stats.cpp b/Engine/Source/Runtime/D3D12RHI/Private/D3D12Stats.cpp
--- a/Engine/Source/Runtime/D3D12RHI/Private/D3D12Stats.cpp
+++ b/Engine/Source/Runtime/D3D12RHI/Private/D3D12Stats.cpp
@@ -38,7 +38,8 @@ void D3D12BufferStats::UpdateBufferStats(FD3D12Buffer& Buffer, bool bAllocating)
 	INC_MEMORY_STAT_BY_FName(GetD3D12BufferStat(BufferDesc).GetName(), RequestedSize);
 	INC_MEMORY_STAT_BY(STAT_D3D12MemoryCurrentTotal, RequestedSize);
 
-#if UE_MEMORY_TRACE_ENABLED
+    // With D3D12RHI_PLATFORM_HAS_UNIFIED_MEMORY=1, MemoryTrace_Alloc is called during resource allocation
+#if UE_MEMORY_TRACE_ENABLED  && (D3D12RHI_PLATFORM_HAS_UNIFIED_MEMORY == 0)
 	const D3D12_GPU_VIRTUAL_ADDRESS GPUAddress = Location.GetGPUVirtualAddress();
 
 	if (bAllocating)
