Replace const TArrayView<const Foo> with TConstArrayView in RHI ray tracing code.

diff --git a/Engine/Source/Runtime/Apple/MetalRHI/Private/MetalRayTracing.cpp b/Engine/Source/Runtime/Apple/MetalRHI/Private/MetalRayTracing.cpp
--- a/Engine/Source/Runtime/Apple/MetalRHI/Private/MetalRayTracing.cpp
+++ b/Engine/Source/Runtime/Apple/MetalRHI/Private/MetalRayTracing.cpp
@@ -549,7 +549,7 @@ void FMetalRHICommandContext::RHIBuildAccelerationStructure(const FRayTracingSce
 		InstanceBuffer, SceneBuildParams.InstanceBufferOffset);
 }
 
-void FMetalRHICommandContext::RHIBuildAccelerationStructures(const TArrayView<const FRayTracingGeometryBuildParams> Params, const FRHIBufferRange& ScratchBufferRange)
+void FMetalRHICommandContext::RHIBuildAccelerationStructures(TConstArrayView<FRayTracingGeometryBuildParams> Params, const FRHIBufferRange& ScratchBufferRange)
 {
 	checkf(ScratchBufferRange.Buffer != nullptr, TEXT("BuildAccelerationStructures requires valid scratch buffer"));
 
diff --git a/Engine/Source/Runtime/Apple/MetalRHI/Public/MetalRHIContext.h b/Engine/Source/Runtime/Apple/MetalRHI/Public/MetalRHIContext.h
--- a/Engine/Source/Runtime/Apple/MetalRHI/Public/MetalRHIContext.h
+++ b/Engine/Source/Runtime/Apple/MetalRHI/Public/MetalRHIContext.h
@@ -149,7 +149,7 @@ public:
 
 #if METAL_RHI_RAYTRACING
 	virtual void RHIBindAccelerationStructureMemory(FRHIRayTracingScene* Scene, FRHIBuffer* Buffer, uint32 BufferOffset) final override;
-	virtual void RHIBuildAccelerationStructures(const TArrayView<const FRayTracingGeometryBuildParams> Params, const FRHIBufferRange& ScratchBufferRange) final override;
+	virtual void RHIBuildAccelerationStructures(TConstArrayView<FRayTracingGeometryBuildParams> Params, const FRHIBufferRange& ScratchBufferRange) final override;
 	virtual void RHIBuildAccelerationStructure(const FRayTracingSceneBuildParams& SceneBuildParams) final override;
 	virtual void RHIClearRayTracingBindings(FRHIRayTracingScene* Scene) final override;
 	virtual void RHIClearShaderBindingTable(FRHIShaderBindingTable* SBT) final override;
diff --git a/Engine/Source/Runtime/D3D12RHI/Private/D3D12CommandContext.h b/Engine/Source/Runtime/D3D12RHI/Private/D3D12CommandContext.h
--- a/Engine/Source/Runtime/D3D12RHI/Private/D3D12CommandContext.h
+++ b/Engine/Source/Runtime/D3D12RHI/Private/D3D12CommandContext.h
@@ -720,10 +720,10 @@ public:
 
 #if D3D12_RHI_RAYTRACING
 	virtual void RHIBindAccelerationStructureMemory(FRHIRayTracingScene* Scene, FRHIBuffer* Buffer, uint32 BufferOffset) final override;
-	virtual void BuildAccelerationStructuresInternal(const TArrayView<D3D12_BUILD_RAYTRACING_ACCELERATION_STRUCTURE_DESC> BuildDesc);
+	virtual void BuildAccelerationStructuresInternal(TConstArrayView<D3D12_BUILD_RAYTRACING_ACCELERATION_STRUCTURE_DESC> BuildDesc);
 #if WITH_MGPU
 	// Should be called before RHIBuildAccelerationStructures when multiple GPU support is present (for example, from FD3D12CommandContextRedirector::RHIBuildAccelerationStructures)
-	static void UnregisterAccelerationStructuresInternalMGPU(const TArrayView<const FRayTracingGeometryBuildParams> Params, FRHIGPUMask GPUMask);
+	static void UnregisterAccelerationStructuresInternalMGPU(TConstArrayView<FRayTracingGeometryBuildParams> Params, FRHIGPUMask GPUMask);
 #endif
 	virtual void RHIBuildAccelerationStructures(TConstArrayView<FRayTracingGeometryBuildParams> Params, const FRHIBufferRange& ScratchBufferRange) final override;
 	virtual void RHIBuildAccelerationStructure(const FRayTracingSceneBuildParams& SceneBuildParams) final override;
diff --git a/Engine/Source/Runtime/D3D12RHI/Private/D3D12RayTracing.cpp b/Engine/Source/Runtime/D3D12RHI/Private/D3D12RayTracing.cpp
--- a/Engine/Source/Runtime/D3D12RHI/Private/D3D12RayTracing.cpp
+++ b/Engine/Source/Runtime/D3D12RHI/Private/D3D12RayTracing.cpp
@@ -4651,7 +4651,7 @@ FRHIShaderBindingTable* FD3D12RayTracingScene::FindOrCreateShaderBindingTable(co
 	return CreatedShaderTable;
 }
 
-void FD3D12CommandContext::BuildAccelerationStructuresInternal(const TArrayView<D3D12_BUILD_RAYTRACING_ACCELERATION_STRUCTURE_DESC> BuildDescs)
+void FD3D12CommandContext::BuildAccelerationStructuresInternal(TConstArrayView<D3D12_BUILD_RAYTRACING_ACCELERATION_STRUCTURE_DESC> BuildDescs)
 {
 	for (const D3D12_BUILD_RAYTRACING_ACCELERATION_STRUCTURE_DESC& Desc : BuildDescs)
 	{
@@ -4660,7 +4660,7 @@ void FD3D12CommandContext::BuildAccelerationStructuresInternal(const TArrayView<
 }
 
 #if WITH_MGPU
-void FD3D12CommandContext::UnregisterAccelerationStructuresInternalMGPU(const TArrayView<const FRayTracingGeometryBuildParams> Params, FRHIGPUMask GPUMask)
+void FD3D12CommandContext::UnregisterAccelerationStructuresInternalMGPU(TConstArrayView<FRayTracingGeometryBuildParams> Params, FRHIGPUMask GPUMask)
 {
 	// We need to unregister rename listeners for all GPUs in a separate pass before running "RHIBuildAccelerationStructures", as the build process
 	// may modify the buffer references in the ray tracing geometry.  This leads to an assert where the code attempts to unregister the newer buffer
@@ -4677,7 +4677,7 @@ void FD3D12CommandContext::UnregisterAccelerationStructuresInternalMGPU(const TA
 }
 #endif  // WITH_MGPU
 
-void FD3D12CommandContext::RHIBuildAccelerationStructures(const TArrayView<const FRayTracingGeometryBuildParams> Params, const FRHIBufferRange& ScratchBufferRange)
+void FD3D12CommandContext::RHIBuildAccelerationStructures(TConstArrayView<FRayTracingGeometryBuildParams> Params, const FRHIBufferRange& ScratchBufferRange)
 {
 	TRACE_CPUPROFILER_EVENT_SCOPE(BuildAccelerationStructure_BottomLevel);
 	SCOPE_CYCLE_COUNTER(STAT_D3D12BuildBLAS);
diff --git a/Engine/Source/Runtime/RHI/Private/RHICommandList.cpp b/Engine/Source/Runtime/RHI/Private/RHICommandList.cpp
--- a/Engine/Source/Runtime/RHI/Private/RHICommandList.cpp
+++ b/Engine/Source/Runtime/RHI/Private/RHICommandList.cpp
@@ -1770,7 +1770,7 @@ void FRHIComputeCommandList::BuildAccelerationStructure(FRHIRayTracingGeometry*
 	BuildAccelerationStructures(MakeArrayView(&Params, 1), ScratchBufferRange);
 }
 
-void FRHIComputeCommandList::BuildAccelerationStructures(const TArrayView<const FRayTracingGeometryBuildParams> Params)
+void FRHIComputeCommandList::BuildAccelerationStructures(TConstArrayView<FRayTracingGeometryBuildParams> Params)
 {
 	// Buffer size is limited to 2Gb, thus split acceleration structure building into pieces to accommodate this limitation
 	const uint64 MaxScratchMemorySize = FMath::Min(uint32(CVarRHICmdMaxAccelerationStructureBuildScratchSize.GetValueOnRenderThread()), 2147483647u);
@@ -1820,7 +1820,7 @@ void FRHIComputeCommandList::BuildAccelerationStructures(const TArrayView<const
 		// Allocate scratch buffer and build the acceleration structure for the selected sub-range
 		if (ParamCount > 0)
 		{
-			const TArrayView<const FRayTracingGeometryBuildParams> EffectiveParams = MakeArrayView<const FRayTracingGeometryBuildParams>(&Params[ParamBegin], ParamCount);
+			TConstArrayView<FRayTracingGeometryBuildParams> EffectiveParams = MakeConstArrayView<FRayTracingGeometryBuildParams>(&Params[ParamBegin], ParamCount);
 			BuildAccelerationStructures(EffectiveParams, ScratchBufferRange);
 		}
 	}
diff --git a/Engine/Source/Runtime/VulkanRHI/Private/VulkanContext.h b/Engine/Source/Runtime/VulkanRHI/Private/VulkanContext.h
--- a/Engine/Source/Runtime/VulkanRHI/Private/VulkanContext.h
+++ b/Engine/Source/Runtime/VulkanRHI/Private/VulkanContext.h
@@ -313,7 +313,7 @@ public:
 	virtual void RHIClearShaderBindingTable(FRHIShaderBindingTable* SBT) final override;
 	virtual void RHICommitShaderBindingTable(FRHIShaderBindingTable* SBT, FRHIBuffer* InlineBindingDataBuffer) final override;
 	virtual void RHIBindAccelerationStructureMemory(FRHIRayTracingScene* Scene, FRHIBuffer* Buffer, uint32 BufferOffset) final override;
-	virtual void RHIBuildAccelerationStructures(const TArrayView<const FRayTracingGeometryBuildParams> Params, const FRHIBufferRange& ScratchBufferRange) final override;
+	virtual void RHIBuildAccelerationStructures(TConstArrayView<FRayTracingGeometryBuildParams> Params, const FRHIBufferRange& ScratchBufferRange) final override;
 	virtual void RHIBuildAccelerationStructure(const FRayTracingSceneBuildParams& SceneBuildParams) final override;
 
 	virtual void RHIRayTraceDispatch(FRHIRayTracingPipelineState* RayTracingPipelineState, FRHIRayTracingShader* RayGenShader,
diff --git a/Engine/Source/Runtime/VulkanRHI/Private/VulkanRayTracing.cpp b/Engine/Source/Runtime/VulkanRHI/Private/VulkanRayTracing.cpp
--- a/Engine/Source/Runtime/VulkanRHI/Private/VulkanRayTracing.cpp
+++ b/Engine/Source/Runtime/VulkanRHI/Private/VulkanRayTracing.cpp
@@ -202,7 +202,7 @@ static ERayTracingAccelerationStructureFlags GetRayTracingAccelerationStructureB
 
 static void GetBLASBuildData(
 	const VkDevice Device,
-	const TArrayView<const FRayTracingGeometrySegment> Segments,
+	const TConstArrayView<FRayTracingGeometrySegment> Segments,
 	const ERayTracingGeometryType GeometryType,
 	const FBufferRHIRef IndexBufferRHI,
 	const uint32 IndexBufferOffset,
@@ -1225,7 +1225,7 @@ void FVulkanCommandListContext::RHIBindAccelerationStructureMemory(FRHIRayTracin
 }
 
 // Todo: High level rhi call should have transitioned and verified vb and ib to read for each segment
-void FVulkanCommandListContext::RHIBuildAccelerationStructures(const TArrayView<const FRayTracingGeometryBuildParams> Params, const FRHIBufferRange& ScratchBufferRange)
+void FVulkanCommandListContext::RHIBuildAccelerationStructures(TConstArrayView<FRayTracingGeometryBuildParams> Params, const FRHIBufferRange& ScratchBufferRange)
 {
 	checkf(ScratchBufferRange.Buffer != nullptr, TEXT("BuildAccelerationStructures requires valid scratch buffer"));
 
