Completely clear known nVidia driver PSO cache folder when clearPSODriveCache is used because file extension has been changed

diff --git a/Engine/Source/Runtime/D3D12RHI/Private/Windows/WindowsD3D12Device.cpp b/Engine/Source/Runtime/D3D12RHI/Private/Windows/WindowsD3D12Device.cpp
--- a/Engine/Source/Runtime/D3D12RHI/Private/Windows/WindowsD3D12Device.cpp
+++ b/Engine/Source/Runtime/D3D12RHI/Private/Windows/WindowsD3D12Device.cpp
@@ -1302,12 +1302,10 @@ static void ClearPSODriverCache()
 		{
 			// NVIDIA used to have a global cache, but now also has a per-driver cache in a different folder in LocalLow.
 			FString GlobalPSOPath = FPaths::Combine(*LocalAppDataFolder, TEXT("NVIDIA"), TEXT("DXCache"));
-			ClearFolder(GlobalPSOPath, TEXT(".bin"));
-			ClearFolder(GlobalPSOPath, TEXT(".toc"));
+			ClearFolder(GlobalPSOPath, nullptr);
 
 			FString PerDriverPSOPath = FPaths::Combine(*LocalAppDataFolder, TEXT(".."), TEXT("LocalLow"), TEXT("NVIDIA"), TEXT("PerDriverVersion"), TEXT("DXCache"));
-			ClearFolder(PerDriverPSOPath, TEXT(".bin"));
-			ClearFolder(PerDriverPSOPath, TEXT(".toc"));
+			ClearFolder(PerDriverPSOPath, nullptr);
 		}
 		else if (IsRHIDeviceAMD())
 		{
