Fix error C4723: potential divide by 0 while compiling FD3D12DefaultBufferAllocator::UpdateMemoryStats (preview compiler)

#rnx

diff --git a/Engine/Source/Runtime/D3D12RHI/Private/D3D12Allocation.cpp b/Engine/Source/Runtime/D3D12RHI/Private/D3D12Allocation.cpp
--- a/Engine/Source/Runtime/D3D12RHI/Private/D3D12Allocation.cpp
+++ b/Engine/Source/Runtime/D3D12RHI/Private/D3D12Allocation.cpp
@@ -1670,7 +1670,8 @@ void FD3D12DefaultBufferAllocator::UpdateMemoryStats()
 
 	// compute fragmentation percentage stats:
 	uint32 Fragmentation = FreeMemory - EndFreeMemory;
-	float FragmentationPercentage = float(Fragmentation) / float(MemoryUsed + AlignmentWaste + Fragmentation);
+	uint32 FragmentationUsed = MemoryUsed + AlignmentWaste + Fragmentation;
+	float FragmentationPercentage = FragmentationUsed == 0 ? 0.0 : float(Fragmentation) / float(FragmentationUsed);
 
 	SET_MEMORY_STAT(STAT_D3D12BufferPoolMemoryAllocated, MemoryAllocated);
 	SET_MEMORY_STAT(STAT_D3D12BufferPoolMemoryUsed, MemoryUsed);
