- Track max used hit and inline record index per worker thread and merge during commit (fix possible race on shared data)
- Remove unused SRV/UAV transition code

diff --git a/Engine/Source/Runtime/D3D12RHI/Private/D3D12RayTracing.cpp b/Engine/Source/Runtime/D3D12RHI/Private/D3D12RayTracing.cpp
--- a/Engine/Source/Runtime/D3D12RHI/Private/D3D12RayTracing.cpp
+++ b/Engine/Source/Runtime/D3D12RHI/Private/D3D12RayTracing.cpp
@@ -1436,7 +1436,8 @@ public:
 	~FD3D12RayTracingShaderBindingTableInternal()
 	{
 		// Might be slow - iterate faster somehow?
-		for (int32 RecordIndex = 0; RecordIndex <= MaxUsedHitRecordIndex; ++RecordIndex)
+		int32 MaxUsedRecordIndex = (int32)FMath::Max(WorkerData[0].MaxUsedHitRecordIndex, WorkerData[0].MaxUsedInlineGeometryParameterSegmentIndex);
+		for (int32 RecordIndex = 0; RecordIndex <= MaxUsedRecordIndex; ++RecordIndex)
 		{
 			ClearHitRecordData(0, RecordIndex);
 		}
@@ -1547,7 +1548,7 @@ public:
 		CompareLocalShaderParameters(ShaderTableOffset, RecordIndex, OffsetWithinRootSignature, SystemParameters);
 	}
 
-	void SetHitGroupIdentifier(uint32 RecordIndex, const FD3D12ShaderIdentifier& ShaderIdentifier)
+	void SetHitGroupIdentifier(uint32 WorkerIndex, uint32 RecordIndex, const FD3D12ShaderIdentifier& ShaderIdentifier)
 	{
 		checkfSlow(ShaderIdentifier.IsValid(), TEXT("Shader identifier must be initialized FD3D12RayTracingPipelineState::GetShaderIdentifier() before use."));
 		checkSlow(sizeof(ShaderIdentifier.Data) >= ShaderIdentifierSize);
@@ -1555,7 +1556,7 @@ public:
 		const uint32 WriteOffset = HitGroupShaderTableOffset + RecordIndex * LocalRecordStride;
 		WriteData(WriteOffset, ShaderIdentifier.Data, ShaderIdentifierSize);
 
-		MaxUsedHitRecordIndex = FMath::Max(MaxUsedHitRecordIndex, (int32)RecordIndex);
+		WorkerData[WorkerIndex].MaxUsedHitRecordIndex = FMath::Max(WorkerData[WorkerIndex].MaxUsedHitRecordIndex, (int32)RecordIndex);
 	}	
 
 	void CompareHitGroupIdentifier(uint32 RecordIndex, const FD3D12ShaderIdentifier& ShaderIdentifier)
@@ -1588,12 +1589,12 @@ public:
 		}
 	}
 
-	void SetInlineGeometryParameters(uint32 SegmentIndex, const void* InData, uint32 InDataSize)
+	void SetInlineGeometryParameters(uint32 WorkerIndex, uint32 SegmentIndex, const void* InData, uint32 InDataSize)
 	{
 		const uint32 WriteOffset = InDataSize * SegmentIndex;
 		FMemory::Memcpy(&InlineGeometryParameterData[WriteOffset], InData, InDataSize);
 				
-		MaxUsedInlineGeometryParameterSegmentIndex = FMath::Max(MaxUsedInlineGeometryParameterSegmentIndex, (int32)SegmentIndex);
+		WorkerData[WorkerIndex].MaxUsedInlineGeometryParameterSegmentIndex = FMath::Max(WorkerData[WorkerIndex].MaxUsedInlineGeometryParameterSegmentIndex, (int32)SegmentIndex);
 	}
 
 	void CompareInlineGeometryParameters(uint32 SegmentIndex, const void* InData, uint32 InDataSize)
@@ -1615,7 +1616,7 @@ public:
 		check(IsInRHIThread() || !IsRunningRHIInSeparateThread());
 		checkf(bIsDirty, TEXT("bIsDirty should be checked before calling Commit()"));
 
-		#if ENABLE_RESIDENCY_MANAGEMENT
+	#if ENABLE_RESIDENCY_MANAGEMENT
 		// Merge all data from worker threads into the main set
 		{
 			TRACE_CPUPROFILER_EVENT_SCOPE(MergeWorkerData);
@@ -1655,27 +1656,19 @@ public:
 			WorkerData[0].DynamicReferencedResourceSet.Empty();
 			WorkerData[0].LastTrackedBaseShaderResources.Empty();
 		}
-		#endif // ENABLE_RESIDENCY_MANAGEMENT
+	#endif // ENABLE_RESIDENCY_MANAGEMENT
 
-		// Merge all transitions and perform other non residency related per worker thread
+		// Merge max record index from worker data into index 0
 		for (uint32 WorkerIndex = 1; WorkerIndex < MaxBindingWorkers; ++WorkerIndex)
 		{
-			for (FD3D12ShaderResourceView* SRV : WorkerData[WorkerIndex].TransitionSRVs)
-			{
-				AddResourceTransition(SRV, 0);
-			}
+			WorkerData[0].MaxUsedHitRecordIndex = FMath::Max(WorkerData[0].MaxUsedHitRecordIndex, WorkerData[WorkerIndex].MaxUsedHitRecordIndex);
+			WorkerData[0].MaxUsedInlineGeometryParameterSegmentIndex = FMath::Max(WorkerData[0].MaxUsedInlineGeometryParameterSegmentIndex, WorkerData[WorkerIndex].MaxUsedInlineGeometryParameterSegmentIndex);
 
-			for (FD3D12UnorderedAccessView* UAV : WorkerData[WorkerIndex].TransitionUAVs)
-			{
-				AddResourceTransition(UAV, 0);
-			}
-
-			WorkerData[WorkerIndex].TransitionSRVs.Empty();
-			WorkerData[WorkerIndex].TransitionUAVs.Empty();
-			WorkerData[WorkerIndex].TransitionViewSet.Empty();
+			WorkerData[WorkerIndex].MaxUsedHitRecordIndex = 0;
+			WorkerData[WorkerIndex].MaxUsedInlineGeometryParameterSegmentIndex = 0;
 		}
 
-		// Clear the record cache for reuse for all worker thread
+		// Clear the record cache for reuse for all worker threads (including index 0)
 		for (uint32 WorkerIndex = 0; WorkerIndex < MaxBindingWorkers; ++WorkerIndex)
 		{
 			WorkerData[WorkerIndex].ShaderRecordCache.Empty();
@@ -1724,13 +1717,13 @@ public:
 		}
 
 		// Also copy geometry parameter data to the GPU buffer
-		if (EnumHasAnyFlags(ShaderBindingMode, ERayTracingShaderBindingMode::Inline) && MaxUsedInlineGeometryParameterSegmentIndex >= 0)
+		if (EnumHasAnyFlags(ShaderBindingMode, ERayTracingShaderBindingMode::Inline) && WorkerData[0].MaxUsedInlineGeometryParameterSegmentIndex >= 0)
 		{
 			check(InInlineBindingDataBuffer);
 
 			TRHICommandList_RecursiveHazardous<FD3D12CommandContext> RHICmdList(&Context);
 			RHICmdList.SetAllowExtraTransition(false);
-			uint32 MaxUsedInlineGeometryParamaterDataSize = (MaxUsedInlineGeometryParameterSegmentIndex + 1) * sizeof(FD3D12HitGroupSystemParameters);
+			uint32 MaxUsedInlineGeometryParamaterDataSize = (WorkerData[0].MaxUsedInlineGeometryParameterSegmentIndex + 1) * sizeof(FD3D12HitGroupSystemParameters);
 			void* MappedUploadMemory = GDynamicRHI->RHILockBuffer(RHICmdList, InInlineBindingDataBuffer, 0, MaxUsedInlineGeometryParamaterDataSize, EResourceLockMode::RLM_WriteOnly);
 			FMemory::Memcpy(MappedUploadMemory, InlineGeometryParameterData.GetData(), MaxUsedInlineGeometryParamaterDataSize);
 			GDynamicRHI->RHIUnlockBuffer(RHICmdList, InInlineBindingDataBuffer);
@@ -1802,9 +1795,6 @@ public:
 	uint32 NumCallableRecords = 0;
 	uint32 NumMissRecords = 0;
 
-	int32 MaxUsedHitRecordIndex = -1;
-	int32 MaxUsedInlineGeometryParameterSegmentIndex = -1;
-
 	uint32 MissShaderTableOffset = 0;
 	uint32 HitGroupShaderTableOffset = 0;
 	uint32 CallableShaderTableOffset = 0;
@@ -2491,37 +2481,12 @@ public:
 		CommandContext.UpdateResidency(Buffer->GetResource());
 	}
 
-	void AddResourceTransition(FD3D12ShaderResourceView* SRV, uint32 WorkerIndex)
-	{
-		bool bAlreadyInSet = false;
-		WorkerData[WorkerIndex].TransitionViewSet.Add(SRV, &bAlreadyInSet);
-		if (!bAlreadyInSet)
-		{
-			WorkerData[WorkerIndex].TransitionSRVs.Add(SRV);
-		}
-	}
-
-	void AddResourceTransition(FD3D12UnorderedAccessView* UAV, uint32 WorkerIndex)
-	{
-		bool bAlreadyInSet = false;
-		WorkerData[WorkerIndex].TransitionViewSet.Add(UAV, &bAlreadyInSet);
-
-		if (!bAlreadyInSet)
-		{
-			WorkerData[WorkerIndex].TransitionUAVs.Add(UAV);
-		}
-	}
-
-
 	struct alignas(PLATFORM_CACHE_LINE_SIZE) FWorkerThreadData
 	{
 		Experimental::TSherwoodMap<FShaderRecordCacheKey, uint32> ShaderRecordCache;
 
-		// Some resources referenced in SBT may be dynamic (written on GPU timeline) and may require transition barriers.
-		// We save such resources while we fill the SBT and issue transitions before the SBT is used.
-		Experimental::TSherwoodSet<FD3D12View*> TransitionViewSet;
-		TArray<FD3D12ShaderResourceView*> TransitionSRVs;
-		TArray<FD3D12UnorderedAccessView*> TransitionUAVs;
+		int32 MaxUsedHitRecordIndex = -1;
+		int32 MaxUsedInlineGeometryParameterSegmentIndex = -1;
 
 #if ENABLE_RESIDENCY_MANAGEMENT
 		// Track all persistent resources with number of refs accross all records (when ref reaches 0 then it doesn't need be tracked anymore)
@@ -5076,22 +5041,6 @@ struct FD3D12RayTracingLocalResourceBinder
 		ShaderTable.AddReferencedUniformBuffer(WorkerIndex, RecordIndex, ShaderTableOffset, OffsetWithinRootSignature, BindingType, UniformBuffer);
 	}
 
-	void AddResourceTransition(FD3D12ShaderResourceView* SRV)
-	{
-		if (SRV->GetResource()->RequiresResourceStateTracking() && ShaderTable.Lifetime == ERayTracingShaderBindingTableLifetime::Transient)
-		{
-			ShaderTable.AddResourceTransition(SRV, WorkerIndex);
-		}
-	}
-
-	void AddResourceTransition(FD3D12UnorderedAccessView* UAV)
-	{
-		if (UAV->GetResource()->RequiresResourceStateTracking() && ShaderTable.Lifetime == ERayTracingShaderBindingTableLifetime::Transient)
-		{
-			ShaderTable.AddResourceTransition(UAV, WorkerIndex);
-		}
-	}
-
 	void AddRayTracingSceneReference(FD3D12RayTracingScene* D3D12RayTracingScene)
 	{
 		checkf(false, TEXT("Unexpected RayTracingScene reference in local shader bindings"));
@@ -5813,7 +5762,7 @@ static void SetRayTracingHitGroup(
 	if (ShaderTable->HitGroupIndexingMode == ERayTracingHitGroupIndexingMode::Disallow)
 	{
 		check(RecordIndex == 0);
-		ShaderTable->SetHitGroupIdentifier(RecordIndex, Pipeline->HitGroupShaders.Identifiers[HitGroupIndex]);
+		ShaderTable->SetHitGroupIdentifier(WorkerIndex, RecordIndex, Pipeline->HitGroupShaders.Identifiers[HitGroupIndex]);
 		return;
 	}
 	
@@ -5869,7 +5818,7 @@ static void SetRayTracingHitGroup(
 			{
 				// Simply copy local shader parameters from existing SBT record and set the shader identifier, skipping resource binding work.
 				const uint32 OffsetFromRootSignatureStart = sizeof(FD3D12HitGroupSystemParameters);
-				ShaderTable->SetHitGroupIdentifier(RecordIndex, Pipeline->HitGroupShaders.Identifiers[HitGroupIndex]);
+				ShaderTable->SetHitGroupIdentifier(WorkerIndex, RecordIndex, Pipeline->HitGroupShaders.Identifiers[HitGroupIndex]);
 				ShaderTable->CopyHitGroupParameters(RecordIndex, *ExistingRecordIndex, OffsetFromRootSignatureStart);
 
 				#if DO_CHECK
@@ -5900,7 +5849,7 @@ static void SetRayTracingHitGroup(
 		FD3D12ShaderIdentifier ShaderIdentifier = bResourcesBound ? Pipeline->HitGroupShaders.Identifiers[HitGroupIndex] : FD3D12ShaderIdentifier::Null;
 		if (RTPSOBindingType != ERayTracingLocalShaderBindingType::Validation)
 		{
-			ShaderTable->SetHitGroupIdentifier(RecordIndex, ShaderIdentifier);
+			ShaderTable->SetHitGroupIdentifier(WorkerIndex, RecordIndex, ShaderIdentifier);
 		}
 		else
 		{
@@ -5918,7 +5867,7 @@ static void SetRayTracingHitGroup(
 			uint32 SegmentIndex = RecordIndex / NumShaderSlotsPerGeometrySegment;
 			if (BindingType != ERayTracingLocalShaderBindingType::Validation)
 			{
-				ShaderTable->SetInlineGeometryParameters(SegmentIndex, &SystemParameters, sizeof(FD3D12HitGroupSystemParameters));
+				ShaderTable->SetInlineGeometryParameters(WorkerIndex, SegmentIndex, &SystemParameters, sizeof(FD3D12HitGroupSystemParameters));
 			}
 			else
 			{
