Also register inline SBT records as listener for raytracing geometry changes because they also stored the bindless index & vertex buffer data which can be modified due to BLAS streaming

diff --git a/Engine/Source/Runtime/D3D12RHI/Private/D3D12RayTracing.cpp b/Engine/Source/Runtime/D3D12RHI/Private/D3D12RayTracing.cpp
--- a/Engine/Source/Runtime/D3D12RHI/Private/D3D12RayTracing.cpp
+++ b/Engine/Source/Runtime/D3D12RHI/Private/D3D12RayTracing.cpp
@@ -1436,8 +1436,7 @@ public:
 	~FD3D12RayTracingShaderBindingTableInternal()
 	{
 		// Might be slow - iterate faster somehow?
-		int32 MaxUsedRecordIndex = (int32)FMath::Max(WorkerData[0].MaxUsedHitRecordIndex, WorkerData[0].MaxUsedInlineGeometryParameterSegmentIndex);
-		for (int32 RecordIndex = 0; RecordIndex <= MaxUsedRecordIndex; ++RecordIndex)
+		for (int32 RecordIndex = 0; RecordIndex <= WorkerData[0].MaxUsedHitRecordIndex; ++RecordIndex)
 		{
 			ClearHitRecordData(0, RecordIndex);
 		}
@@ -1548,15 +1547,13 @@ public:
 		CompareLocalShaderParameters(ShaderTableOffset, RecordIndex, OffsetWithinRootSignature, SystemParameters);
 	}
 
-	void SetHitGroupIdentifier(uint32 WorkerIndex, uint32 RecordIndex, const FD3D12ShaderIdentifier& ShaderIdentifier)
+	void SetHitGroupIdentifier(uint32 RecordIndex, const FD3D12ShaderIdentifier& ShaderIdentifier)
 	{
 		checkfSlow(ShaderIdentifier.IsValid(), TEXT("Shader identifier must be initialized FD3D12RayTracingPipelineState::GetShaderIdentifier() before use."));
 		checkSlow(sizeof(ShaderIdentifier.Data) >= ShaderIdentifierSize);
 
 		const uint32 WriteOffset = HitGroupShaderTableOffset + RecordIndex * LocalRecordStride;
 		WriteData(WriteOffset, ShaderIdentifier.Data, ShaderIdentifierSize);
-
-		WorkerData[WorkerIndex].MaxUsedHitRecordIndex = FMath::Max(WorkerData[WorkerIndex].MaxUsedHitRecordIndex, (int32)RecordIndex);
 	}	
 
 	void CompareHitGroupIdentifier(uint32 RecordIndex, const FD3D12ShaderIdentifier& ShaderIdentifier)
@@ -1589,19 +1586,19 @@ public:
 		}
 	}
 
-	void SetInlineGeometryParameters(uint32 WorkerIndex, uint32 SegmentIndex, const void* InData, uint32 InDataSize)
+	void SetInlineGeometryParameters(uint32 InlineRecordIndex, const FD3D12HitGroupSystemParameters& SystemParameters)
 	{
-		const uint32 WriteOffset = InDataSize * SegmentIndex;
-		FMemory::Memcpy(&InlineGeometryParameterData[WriteOffset], InData, InDataSize);
-				
-		WorkerData[WorkerIndex].MaxUsedInlineGeometryParameterSegmentIndex = FMath::Max(WorkerData[WorkerIndex].MaxUsedInlineGeometryParameterSegmentIndex, (int32)SegmentIndex);
+		const uint32 DataSize = sizeof(FD3D12HitGroupSystemParameters);
+		const uint32 WriteOffset = DataSize * InlineRecordIndex;
+		FMemory::Memcpy(&InlineGeometryParameterData[WriteOffset], &SystemParameters, DataSize);
 	}
 
-	void CompareInlineGeometryParameters(uint32 SegmentIndex, const void* InData, uint32 InDataSize)
+	void CompareInlineGeometryParameters(uint32 InlineRecordIndex, const FD3D12HitGroupSystemParameters& SystemParameters)
 	{
-		const uint32 WriteOffset = InDataSize * SegmentIndex;
+		const uint32 DataSize = sizeof(FD3D12HitGroupSystemParameters);
+		const uint32 WriteOffset = DataSize * InlineRecordIndex;
 		const void* CurrentData = InlineGeometryParameterData.GetData() + WriteOffset;
-		ensure(FMemory::Memcmp(CurrentData, InData, InDataSize) == 0);
+		ensure(FMemory::Memcmp(CurrentData, &SystemParameters, DataSize) == 0);
 	}
 	   
 	FRHISizeAndStride GetInlineBindingDataSizeAndStride() const
@@ -1662,10 +1659,7 @@ public:
 		for (uint32 WorkerIndex = 1; WorkerIndex < MaxBindingWorkers; ++WorkerIndex)
 		{
 			WorkerData[0].MaxUsedHitRecordIndex = FMath::Max(WorkerData[0].MaxUsedHitRecordIndex, WorkerData[WorkerIndex].MaxUsedHitRecordIndex);
-			WorkerData[0].MaxUsedInlineGeometryParameterSegmentIndex = FMath::Max(WorkerData[0].MaxUsedInlineGeometryParameterSegmentIndex, WorkerData[WorkerIndex].MaxUsedInlineGeometryParameterSegmentIndex);
-
 			WorkerData[WorkerIndex].MaxUsedHitRecordIndex = 0;
-			WorkerData[WorkerIndex].MaxUsedInlineGeometryParameterSegmentIndex = 0;
 		}
 
 		// Clear the record cache for reuse for all worker threads (including index 0)
@@ -1717,13 +1711,15 @@ public:
 		}
 
 		// Also copy geometry parameter data to the GPU buffer
-		if (EnumHasAnyFlags(ShaderBindingMode, ERayTracingShaderBindingMode::Inline) && WorkerData[0].MaxUsedInlineGeometryParameterSegmentIndex >= 0)
+		if (EnumHasAnyFlags(ShaderBindingMode, ERayTracingShaderBindingMode::Inline) && WorkerData[0].MaxUsedHitRecordIndex >= 0)
 		{
 			check(InInlineBindingDataBuffer);
 
+			uint32 MaxUsedInlineRecordIndex = WorkerData[0].MaxUsedHitRecordIndex / NumShaderSlotsPerGeometrySegment;
+
 			TRHICommandList_RecursiveHazardous<FD3D12CommandContext> RHICmdList(&Context);
 			RHICmdList.SetAllowExtraTransition(false);
-			uint32 MaxUsedInlineGeometryParamaterDataSize = (WorkerData[0].MaxUsedInlineGeometryParameterSegmentIndex + 1) * sizeof(FD3D12HitGroupSystemParameters);
+			uint32 MaxUsedInlineGeometryParamaterDataSize = (MaxUsedInlineRecordIndex + 1) * sizeof(FD3D12HitGroupSystemParameters);
 			void* MappedUploadMemory = GDynamicRHI->RHILockBuffer(RHICmdList, InInlineBindingDataBuffer, 0, MaxUsedInlineGeometryParamaterDataSize, EResourceLockMode::RLM_WriteOnly);
 			FMemory::Memcpy(MappedUploadMemory, InlineGeometryParameterData.GetData(), MaxUsedInlineGeometryParamaterDataSize);
 			GDynamicRHI->RHIUnlockBuffer(RHICmdList, InInlineBindingDataBuffer);
@@ -1920,7 +1916,17 @@ public:
 			uint32 GPUIndex = ShaderTable->Device->GetGPUIndex();
 			FD3D12HitGroupSystemParameters SystemParameters = Geometry->HitGroupSystemParameters[GPUIndex][SegmentIndex];
 			SystemParameters.RootConstants.UserData = UserData;
-			ShaderTable->SetHitGroupSystemParameters(RecordIndex, SystemParameters);
+
+			uint32 InlineRecordIndex = ShaderTable->GetInlineRecordIndex(RecordIndex);
+
+			if (EnumHasAnyFlags(ShaderTable->ShaderBindingMode, ERayTracingShaderBindingMode::RTPSO))
+			{
+				ShaderTable->SetHitGroupSystemParameters(RecordIndex, SystemParameters);
+			}
+			if (EnumHasAnyFlags(ShaderTable->ShaderBindingMode, ERayTracingShaderBindingMode::Inline) && InlineRecordIndex != INDEX_NONE)
+			{
+				ShaderTable->SetInlineGeometryParameters(InlineRecordIndex, SystemParameters);
+			}
 		}
 	};
 
@@ -2105,14 +2111,31 @@ public:
 	};
 	TArray<FRecordData> HitRecordData;
 
-	void SetHitGroupSystemParameters(uint32 WorkerIndex, ERayTracingLocalShaderBindingType BindingType, uint32 RecordIndex, const FD3D12RayTracingGeometry* Geometry, uint32 UserData, uint32 SegmentIndex)
+	uint32 GetInlineRecordIndex(uint32 RecordIndex) const
+	{	
+		// Only care about shader slot 0 for inline geometry parameters -> remap the record index
+		return (RecordIndex % NumShaderSlotsPerGeometrySegment == 0) ? RecordIndex / NumShaderSlotsPerGeometrySegment : INDEX_NONE;
+	}
+
+	void SetHitGroupGeometrySystemParameters(uint32 WorkerIndex, ERayTracingLocalShaderBindingType BindingType, uint32 RecordIndex, const FD3D12RayTracingGeometry* Geometry, uint32 UserData, uint32 GeometrySegmentIndex)
 	{	
-		FD3D12HitGroupSystemParameters SystemParameters = Geometry->HitGroupSystemParameters[Device->GetGPUIndex()][SegmentIndex];
-		SystemParameters.RootConstants.UserData = UserData;
+		FD3D12HitGroupSystemParameters SystemParameters = Geometry->HitGroupSystemParameters[Device->GetGPUIndex()][GeometrySegmentIndex];
+		SystemParameters.RootConstants.UserData = UserData;		
+
+		WorkerData[WorkerIndex].MaxUsedHitRecordIndex = FMath::Max(WorkerData[WorkerIndex].MaxUsedHitRecordIndex, (int32)RecordIndex);
+
+		uint32 InlineRecordIndex = GetInlineRecordIndex(RecordIndex);
 
 		if (Lifetime == ERayTracingShaderBindingTableLifetime::Transient)
 		{
-			SetHitGroupSystemParameters(RecordIndex, SystemParameters);
+			if (EnumHasAnyFlags(ShaderBindingMode, ERayTracingShaderBindingMode::RTPSO))
+			{
+				SetHitGroupSystemParameters(RecordIndex, SystemParameters);
+			}
+			if (EnumHasAnyFlags(ShaderBindingMode, ERayTracingShaderBindingMode::Inline) && InlineRecordIndex != INDEX_NONE)
+			{
+				SetInlineGeometryParameters(InlineRecordIndex, SystemParameters);
+			}
 		}
 		else
 		{
@@ -2129,12 +2152,20 @@ public:
 			if (bSetData)
 			{
 				check(RecordData.State != FRecordData::EState::Persistent || BindingType == ERayTracingLocalShaderBindingType::Transient);
-				SetHitGroupSystemParameters(RecordIndex, SystemParameters);
+
+				if (EnumHasAnyFlags(ShaderBindingMode, ERayTracingShaderBindingMode::RTPSO))
+				{
+					SetHitGroupSystemParameters(RecordIndex, SystemParameters);
+				}
+				if (EnumHasAnyFlags(ShaderBindingMode, ERayTracingShaderBindingMode::Inline) && InlineRecordIndex != INDEX_NONE)
+				{
+					SetInlineGeometryParameters(InlineRecordIndex, SystemParameters);
+				}
 
 				// If persistently stored then also listen to hit record updates for stream in/out events of the ray tracing geometry data
 				if (BindingType == ERayTracingLocalShaderBindingType::Persistent)
 				{
-					HitRecordData[RecordIndex].UpdateGeometryListener.Init(this, Geometry, RecordIndex, UserData, SegmentIndex);
+					HitRecordData[RecordIndex].UpdateGeometryListener.Init(this, Geometry, RecordIndex, UserData, GeometrySegmentIndex);
 					RecordData.State = FRecordData::EState::Persistent;
 				}				
 				else // if (RecordData.State == FRecordData::EState::Uninitialized)
@@ -2145,7 +2176,14 @@ public:
 			else
 			{
 				check(RecordData.State == FRecordData::EState::Persistent);
-				CompareHitGroupSystemParameters(RecordIndex, SystemParameters);
+				if (EnumHasAnyFlags(ShaderBindingMode, ERayTracingShaderBindingMode::RTPSO))
+				{
+					CompareHitGroupSystemParameters(RecordIndex, SystemParameters);
+				}
+				if (EnumHasAnyFlags(ShaderBindingMode, ERayTracingShaderBindingMode::Inline) && InlineRecordIndex != INDEX_NONE)
+				{
+					CompareInlineGeometryParameters(InlineRecordIndex, SystemParameters);
+				}
 			}
 		}
 	}
@@ -2486,7 +2524,6 @@ public:
 		Experimental::TSherwoodMap<FShaderRecordCacheKey, uint32> ShaderRecordCache;
 
 		int32 MaxUsedHitRecordIndex = -1;
-		int32 MaxUsedInlineGeometryParameterSegmentIndex = -1;
 
 #if ENABLE_RESIDENCY_MANAGEMENT
 		// Track all persistent resources with number of refs accross all records (when ref reaches 0 then it doesn't need be tracked anymore)
@@ -5762,7 +5799,7 @@ static void SetRayTracingHitGroup(
 	if (ShaderTable->HitGroupIndexingMode == ERayTracingHitGroupIndexingMode::Disallow)
 	{
 		check(RecordIndex == 0);
-		ShaderTable->SetHitGroupIdentifier(WorkerIndex, RecordIndex, Pipeline->HitGroupShaders.Identifiers[HitGroupIndex]);
+		ShaderTable->SetHitGroupIdentifier(RecordIndex, Pipeline->HitGroupShaders.Identifiers[HitGroupIndex]);
 		return;
 	}
 	
@@ -5790,13 +5827,10 @@ static void SetRayTracingHitGroup(
 	}
 #endif // DO_CHECK
 
-	FD3D12HitGroupSystemParameters SystemParameters = Geometry->HitGroupSystemParameters[GPUIndex][GeometrySegmentIndex];
-	SystemParameters.RootConstants.UserData = UserData;
+	ShaderTable->SetHitGroupGeometrySystemParameters(WorkerIndex, BindingType, RecordIndex, Geometry, UserData, GeometrySegmentIndex);
 
 	if (EnumHasAnyFlags(ShaderTable->ShaderBindingMode, ERayTracingShaderBindingMode::RTPSO) && Pipeline)
 	{
-		ShaderTable->SetHitGroupSystemParameters(WorkerIndex, BindingType, RecordIndex, Geometry, UserData, GeometrySegmentIndex);
-
 		const FD3D12RayTracingShader* Shader = Pipeline->HitGroupShaders.Shaders[HitGroupIndex];
 
 		FD3D12RayTracingShaderBindingTableInternal::FShaderRecordCacheKey CacheKey;
@@ -5818,7 +5852,7 @@ static void SetRayTracingHitGroup(
 			{
 				// Simply copy local shader parameters from existing SBT record and set the shader identifier, skipping resource binding work.
 				const uint32 OffsetFromRootSignatureStart = sizeof(FD3D12HitGroupSystemParameters);
-				ShaderTable->SetHitGroupIdentifier(WorkerIndex, RecordIndex, Pipeline->HitGroupShaders.Identifiers[HitGroupIndex]);
+				ShaderTable->SetHitGroupIdentifier(RecordIndex, Pipeline->HitGroupShaders.Identifiers[HitGroupIndex]);
 				ShaderTable->CopyHitGroupParameters(RecordIndex, *ExistingRecordIndex, OffsetFromRootSignatureStart);
 
 				#if DO_CHECK
@@ -5849,32 +5883,13 @@ static void SetRayTracingHitGroup(
 		FD3D12ShaderIdentifier ShaderIdentifier = bResourcesBound ? Pipeline->HitGroupShaders.Identifiers[HitGroupIndex] : FD3D12ShaderIdentifier::Null;
 		if (RTPSOBindingType != ERayTracingLocalShaderBindingType::Validation)
 		{
-			ShaderTable->SetHitGroupIdentifier(WorkerIndex, RecordIndex, ShaderIdentifier);
+			ShaderTable->SetHitGroupIdentifier(RecordIndex, ShaderIdentifier);
 		}
 		else
 		{
 			ShaderTable->CompareHitGroupIdentifier(RecordIndex, ShaderIdentifier);
 		}
 	}
-
-	if (EnumHasAnyFlags(ShaderTable->ShaderBindingMode, ERayTracingShaderBindingMode::Inline))
-	{
-		// Only care about shader slot 0 for inline geometry parameters
-		uint32 NumShaderSlotsPerGeometrySegment = ShaderTable->NumShaderSlotsPerGeometrySegment;
-		if (RecordIndex % NumShaderSlotsPerGeometrySegment == 0)
-		{			
-			// Recompute the geometry segment index from the record index
-			uint32 SegmentIndex = RecordIndex / NumShaderSlotsPerGeometrySegment;
-			if (BindingType != ERayTracingLocalShaderBindingType::Validation)
-			{
-				ShaderTable->SetInlineGeometryParameters(WorkerIndex, SegmentIndex, &SystemParameters, sizeof(FD3D12HitGroupSystemParameters));
-			}
-			else
-			{
-				ShaderTable->CompareInlineGeometryParameters(SegmentIndex, &SystemParameters, sizeof(FD3D12HitGroupSystemParameters));
-			}
-		}
-	}
 }
 
 static void SetRayTracingCallableShader(
