Fix for ASAN on possible nullptr for pipeline state

diff --git a/Engine/Source/Runtime/D3D12RHI/Private/D3D12RayTracing.cpp b/Engine/Source/Runtime/D3D12RHI/Private/D3D12RayTracing.cpp
--- a/Engine/Source/Runtime/D3D12RHI/Private/D3D12RayTracing.cpp
+++ b/Engine/Source/Runtime/D3D12RHI/Private/D3D12RayTracing.cpp
@@ -5798,8 +5798,11 @@ static void SetRayTracingHitGroup(
 	// If Shader table doesn't support hit group indexing then only set the hit group identifier and it should be first record index
 	if (ShaderTable->HitGroupIndexingMode == ERayTracingHitGroupIndexingMode::Disallow)
 	{
-		check(RecordIndex == 0);
-		ShaderTable->SetHitGroupIdentifier(RecordIndex, Pipeline->HitGroupShaders.Identifiers[HitGroupIndex]);
+		check(RecordIndex == 0 && Pipeline);
+		if (Pipeline)
+		{
+			ShaderTable->SetHitGroupIdentifier(RecordIndex, Pipeline->HitGroupShaders.Identifiers[HitGroupIndex]);
+		}
 		return;
 	}
 	
