[LLM] Add asset tags for RHITextures and RHIBuffers
- re-scope streaming manager to separate tag from UObject/EngineMisc and others
- re-scope texture profiler allocations into separate tag from Textures/RenderTargets

#rnx

diff --git a/Engine/Source/Runtime/Apple/MetalRHI/MetalRHI.Build.cs b/Engine/Source/Runtime/Apple/MetalRHI/MetalRHI.Build.cs
--- a/Engine/Source/Runtime/Apple/MetalRHI/MetalRHI.Build.cs
+++ b/Engine/Source/Runtime/Apple/MetalRHI/MetalRHI.Build.cs
@@ -15,7 +15,8 @@ public class MetalRHI : ModuleRules
 				"RHI",
 				"RHICore",
 				"RenderCore",
-				"Projects"
+				"Projects",
+				"TraceLog",
 			}
 			);
 
diff --git a/Engine/Source/Runtime/Apple/MetalRHI/Private/MetalRHIBuffer.cpp b/Engine/Source/Runtime/Apple/MetalRHI/Private/MetalRHIBuffer.cpp
--- a/Engine/Source/Runtime/Apple/MetalRHI/Private/MetalRHIBuffer.cpp
+++ b/Engine/Source/Runtime/Apple/MetalRHI/Private/MetalRHIBuffer.cpp
@@ -13,6 +13,8 @@
 #include "RenderUtils.h"
 #include "MetalLLM.h"
 #include "RHICoreBufferInitializer.h"
+#include "HAL/LowLevelMemStats.h"
+#include "ProfilingDebugging/AssetMetadataTrace.h"
 #include <objc/runtime.h>
 
 #define METAL_POOL_BUFFER_BACKING 1
@@ -609,6 +611,9 @@ static FRHIBufferCreateDesc MetalModifyBufferCreateDesc(const FRHIBufferCreateDe
 FRHIBufferInitializer FMetalDynamicRHI::RHICreateBufferInitializer(FRHICommandListBase& RHICmdList, const FRHIBufferCreateDesc& InCreateDesc)
 {
 	MTL_SCOPED_AUTORELEASE_POOL;
+	LLM_SCOPE_DYNAMIC_STAT_OBJECTPATH_FNAME(CreateDesc.OwnerName, ELLMTagSet::Assets);
+	LLM_SCOPE_DYNAMIC_STAT_OBJECTPATH_FNAME(CreateDesc.GetTraceClassName(), ELLMTagSet::AssetClasses);
+	UE_TRACE_METADATA_SCOPE_ASSET_FNAME(CreateDesc.DebugName, CreateDesc.GetTraceClassName(), CreateDesc.OwnerName);
 
 	const FRHIBufferCreateDesc CreateDesc = MetalModifyBufferCreateDesc(InCreateDesc);
 
diff --git a/Engine/Source/Runtime/Apple/MetalRHI/Private/MetalTexture.cpp b/Engine/Source/Runtime/Apple/MetalRHI/Private/MetalTexture.cpp
--- a/Engine/Source/Runtime/Apple/MetalRHI/Private/MetalTexture.cpp
+++ b/Engine/Source/Runtime/Apple/MetalRHI/Private/MetalTexture.cpp
@@ -15,6 +15,8 @@
 #include "RHILockTracker.h"
 #include "RHICoreStats.h"
 #include "MetalBindlessDescriptors.h"
+#include "HAL/LowLevelMemStats.h"
+#include "ProfilingDebugging/AssetMetadataTrace.h"
 #include <CoreVideo/CVMetalTexture.h>
 
 volatile int64 FMetalSurface::ActiveUploads = 0;
@@ -1393,6 +1395,9 @@ uint32 FMetalDynamicRHI::RHIComputeMemorySize(FRHITexture* TextureRHI)
 FTextureRHIRef FMetalDynamicRHI::RHICreateTexture(FRHICommandListBase& RHICmdList, const FRHITextureCreateDesc& CreateDesc)
 {
     MTL_SCOPED_AUTORELEASE_POOL;
+	LLM_SCOPE_DYNAMIC_STAT_OBJECTPATH_FNAME(CreateDesc.OwnerName, ELLMTagSet::Assets);
+	LLM_SCOPE_DYNAMIC_STAT_OBJECTPATH_FNAME(CreateDesc.GetTraceClassName(), ELLMTagSet::AssetClasses);
+	UE_TRACE_METADATA_SCOPE_ASSET_FNAME(CreateDesc.DebugName, CreateDesc.GetTraceClassName(), CreateDesc.OwnerName);
     return new FMetalSurface(*Device, &RHICmdList, FMetalTextureCreateDesc(*Device, CreateDesc));
 }
 
@@ -1424,7 +1429,12 @@ FTextureRHIRef FMetalDynamicRHI::RHIAsyncReallocateTexture2D(FRHITexture* OldTex
         RHIGetDefaultResourceState(Desc.Flags, false),
         TEXT("RHIAsyncReallocateTexture2D")
     );
-    
+	CreateDesc.SetOwnerName(OldTexture->GetOwnerName());
+
+	LLM_SCOPE_DYNAMIC_STAT_OBJECTPATH_FNAME(CreateDesc.OwnerName, ELLMTagSet::Assets);
+	LLM_SCOPE_DYNAMIC_STAT_OBJECTPATH_FNAME(CreateDesc.GetTraceClassName(), ELLMTagSet::AssetClasses);
+	UE_TRACE_METADATA_SCOPE_ASSET_FNAME(CreateDesc.DebugName, CreateDesc.GetTraceClassName(), CreateDesc.OwnerName);
+
     FMetalSurface* NewTexture = new FMetalSurface(*Device, &RHICmdList, FMetalTextureCreateDesc(*Device, CreateDesc));
 
     // Copy shared mips
diff --git a/Engine/Source/Runtime/D3D12RHI/Private/D3D12Buffer.cpp b/Engine/Source/Runtime/D3D12RHI/Private/D3D12Buffer.cpp
--- a/Engine/Source/Runtime/D3D12RHI/Private/D3D12Buffer.cpp
+++ b/Engine/Source/Runtime/D3D12RHI/Private/D3D12Buffer.cpp
@@ -5,16 +5,11 @@ D3D12Buffer.cpp: D3D Common code for buffers.
 =============================================================================*/
 
 #include "D3D12RHIPrivate.h"
+#include "HAL/LowLevelMemStats.h"
 #include "ProfilingDebugging/AssetMetadataTrace.h"
 #include "RHICoreBufferInitializer.h"
 #include "RHICoreStats.h"
 
-static FName GetRHIBufferClassName(const FName ClassName)
-{
-	const static FLazyName RHIBufferName(TEXT("FRHIBuffer"));
-	return (ClassName == NAME_None) ? RHIBufferName : ClassName;
-}
-
 extern int32 GD3D12BindResourceLabels;
 
 FD3D12Buffer::~FD3D12Buffer()
@@ -219,7 +214,9 @@ FD3D12Buffer* FD3D12Adapter::CreateRHIBuffer(
 	TRACE_CPUPROFILER_EVENT_SCOPE(D3D12RHI::CreateRHIBuffer);
 	SCOPE_CYCLE_COUNTER(STAT_D3D12CreateBufferTime);
 
-	UE_TRACE_METADATA_SCOPE_ASSET_FNAME(FName(CreateDesc.DebugName), GetRHIBufferClassName(CreateDesc.ClassName), CreateDesc.OwnerName);
+	LLM_SCOPE_DYNAMIC_STAT_OBJECTPATH_FNAME(CreateDesc.OwnerName, ELLMTagSet::Assets);
+	LLM_SCOPE_DYNAMIC_STAT_OBJECTPATH_FNAME(CreateDesc.GetTraceClassName(), ELLMTagSet::AssetClasses);
+	UE_TRACE_METADATA_SCOPE_ASSET_FNAME(CreateDesc.DebugName, CreateDesc.GetTraceClassName(), CreateDesc.OwnerName);
 
 	check(InDesc.Width >= CreateDesc.Size);
 
@@ -447,7 +444,9 @@ static FRHIBufferInitializer CreateBufferInitializerForWriting(FRHICommandListBa
 
 FRHIBufferInitializer FD3D12DynamicRHI::RHICreateBufferInitializer(FRHICommandListBase& RHICmdList, const FRHIBufferCreateDesc& CreateDesc)
 {
-	UE_TRACE_METADATA_SCOPE_ASSET_FNAME(CreateDesc.DebugName, GetRHIBufferClassName(CreateDesc.ClassName), CreateDesc.OwnerName);
+	LLM_SCOPE_DYNAMIC_STAT_OBJECTPATH_FNAME(CreateDesc.OwnerName, ELLMTagSet::Assets);
+	LLM_SCOPE_DYNAMIC_STAT_OBJECTPATH_FNAME(CreateDesc.GetTraceClassName(), ELLMTagSet::AssetClasses);
+	UE_TRACE_METADATA_SCOPE_ASSET_FNAME(CreateDesc.DebugName, CreateDesc.GetTraceClassName(), CreateDesc.OwnerName);
 
 	if (CreateDesc.IsNull())
 	{
diff --git a/Engine/Source/Runtime/D3D12RHI/Private/D3D12Texture.cpp b/Engine/Source/Runtime/D3D12RHI/Private/D3D12Texture.cpp
--- a/Engine/Source/Runtime/D3D12RHI/Private/D3D12Texture.cpp
+++ b/Engine/Source/Runtime/D3D12RHI/Private/D3D12Texture.cpp
@@ -9,6 +9,7 @@
 #include "TextureProfiler.h"
 #include "ProfilingDebugging/MemoryTrace.h"
 #include "RHIUtilities.h"
+#include "HAL/LowLevelMemStats.h"
 #include "ProfilingDebugging/AssetMetadataTrace.h"
 #include "RHICoreStats.h"
 
@@ -93,6 +94,10 @@ void FD3D12TextureStats::UpdateD3D12TextureStats(FD3D12Texture& Texture, const F
 	INC_MEMORY_STAT_BY(STAT_D3D12MemoryCurrentTotal, TextureSizeDeltaInBytes);
 
 	D3D12_GPU_VIRTUAL_ADDRESS GPUAddress = Texture.ResourceLocation.GetGPUVirtualAddress();
+	if (GPUAddress == 0)
+	{
+		GPUAddress = reinterpret_cast<D3D12_GPU_VIRTUAL_ADDRESS>(Texture.ResourceLocation.GetAddressForLLMTracking());
+	}
 
 #if UE_MEMORY_TRACE_ENABLED
 	// Textures don't have valid GPUVirtualAddress when IsTrackingAllAllocations() is false, so don't do memory trace in this case.
@@ -616,10 +621,12 @@ FD3D12DynamicRHI::FCreateTextureInternalResult FD3D12DynamicRHI::CreateTextureIn
 {
 #if PLATFORM_WINDOWS
 	TRACE_CPUPROFILER_EVENT_SCOPE(D3D12RHI::CreateTextureInternal);
-	UE_TRACE_METADATA_SCOPE_ASSET_FNAME(FName(InCreateDesc.DebugName), InCreateDesc.GetTraceClassName(), InCreateDesc.OwnerName);
 
 	// Make local copy of create desc because certain flags might be removed before actual texture creation
 	FRHITextureCreateDesc CreateDesc = InCreateDesc;
+	LLM_SCOPE_DYNAMIC_STAT_OBJECTPATH_FNAME(CreateDesc.OwnerName, ELLMTagSet::Assets);
+	LLM_SCOPE_DYNAMIC_STAT_OBJECTPATH_FNAME(CreateDesc.GetTraceClassName(), ELLMTagSet::AssetClasses);
+	UE_TRACE_METADATA_SCOPE_ASSET_FNAME(CreateDesc.DebugName, CreateDesc.GetTraceClassName(), CreateDesc.OwnerName);
 
 	// Virtual textures currently not supported in default D3D12
 	check(!EnumHasAnyFlags(CreateDesc.Flags, TexCreate_Virtual));
@@ -844,11 +851,6 @@ public:
 FTextureRHIRef FD3D12DynamicRHI::RHIAsyncCreateTexture2D(uint32 SizeX, uint32 SizeY, uint8 Format, uint32 NumMips, ETextureCreateFlags Flags, ERHIAccess InResourceState, void** InitialMipData, uint32 NumInitialMips, const TCHAR* DebugName, FGraphEventRef& OutCompletionEvent)
 {	
 	check(GRHISupportsAsyncTextureCreation);
-	
-	const static FName RHIAsyncCreateTexture2DName(TEXT("RHIAsyncCreateTexture2D"));
-	const static FName RHITextureName(TEXT("FRHITexture"));
-	UE_TRACE_METADATA_SCOPE_ASSET_FNAME(RHIAsyncCreateTexture2DName, RHITextureName, NAME_None);
-
 	const ETextureCreateFlags InvalidFlags = TexCreate_RenderTargetable | TexCreate_ResolveTargetable | TexCreate_DepthStencilTargetable | TexCreate_UAV | TexCreate_Presentable | TexCreate_CPUReadback;
 	check(!EnumHasAnyFlags(Flags, InvalidFlags));
 
@@ -1285,6 +1287,7 @@ FTextureRHIRef FD3D12DynamicRHI::AsyncReallocateTexture2D_RenderThread(FRHIComma
 		RHIGetDefaultResourceState(Desc.Flags, false),
 		TEXT("AsyncReallocateTexture2D_RenderThread")
 	);
+	CreateDesc.SetOwnerName(OldTexture->GetOwnerName());
 
 	// Allocate a new texture.
 	FD3D12Texture* NewTexture = CreateD3D12Texture(RHICmdList, CreateDesc);
diff --git a/Engine/Source/Runtime/Engine/Private/ContentStreaming.cpp b/Engine/Source/Runtime/Engine/Private/ContentStreaming.cpp
--- a/Engine/Source/Runtime/Engine/Private/ContentStreaming.cpp
+++ b/Engine/Source/Runtime/Engine/Private/ContentStreaming.cpp
@@ -372,6 +372,7 @@ FStreamingManagerCollection& IStreamingManager::Get()
 {
 	if (StreamingManagerCollection == nullptr)
 	{
+		LLM_SCOPE(ELLMTag::StreamingManager);
 		LLM_TAGSET_SCOPE_CLEAR(ELLMTagSet::Assets);
 		LLM_TAGSET_SCOPE_CLEAR(ELLMTagSet::AssetClasses);
 		UE_TRACE_METADATA_CLEAR_SCOPE();
diff --git a/Engine/Source/Runtime/Engine/Private/Streaming/RenderAssetUpdate.cpp b/Engine/Source/Runtime/Engine/Private/Streaming/RenderAssetUpdate.cpp
--- a/Engine/Source/Runtime/Engine/Private/Streaming/RenderAssetUpdate.cpp
+++ b/Engine/Source/Runtime/Engine/Private/Streaming/RenderAssetUpdate.cpp
@@ -9,6 +9,9 @@ RenderAssetUpdate.cpp: Base class of helpers to stream in and out texture/mesh L
 #include "RenderingThread.h"
 #include "Streaming/TextureStreamingHelpers.h"
 #include "UObject/UObjectIterator.h"
+#include "HAL/LowLevelMemStats.h"
+#include "ProfilingDebugging/AssetMetadataTrace.h"
+#include "UObject/Package.h"
 
 float GStreamingFlushTimeOut = 3.00f;
 static FAutoConsoleVariableRef CVarStreamingFlushTimeOut(
@@ -199,6 +202,10 @@ void FRenderAssetUpdate::Tick(EThreadType InCurrentThread)
 			ETaskState TickResult;
 			do // Iterate as longs as there is progress
 			{
+				LLM_SCOPE_DYNAMIC_STAT_OBJECTPATH(StreamableAsset, ELLMTagSet::Assets);
+				LLM_SCOPE_DYNAMIC_STAT_OBJECTPATH(StreamableAsset->GetClass(), ELLMTagSet::AssetClasses);
+				UE_TRACE_METADATA_SCOPE_ASSET(StreamableAsset, StreamableAsset->GetClass());
+
 				// Only test for suspension the first time and in normal progress.
 				// When cancelled, we want the update to complete without interruptions, allowing reference to be freed.
 				TickResult = TickInternal(InCurrentThread, !bWasAlreadyLocked && !bIsCancelled);
diff --git a/Engine/Source/Runtime/Engine/Private/Streaming/StreamingManagerTexture.cpp b/Engine/Source/Runtime/Engine/Private/Streaming/StreamingManagerTexture.cpp
--- a/Engine/Source/Runtime/Engine/Private/Streaming/StreamingManagerTexture.cpp
+++ b/Engine/Source/Runtime/Engine/Private/Streaming/StreamingManagerTexture.cpp
@@ -26,6 +26,8 @@
 #include "Components/PrimitiveComponent.h"
 #include "Misc/CoreDelegates.h"
 #include "TextureResource.h"
+#include "HAL/LowLevelMemStats.h"
+#include "ProfilingDebugging/MetadataTrace.h"
 
 #if !STATS
 #include "Async/ParallelFor.h"
@@ -302,6 +304,10 @@ void FRenderAssetStreamingManager::MarkMountedStateDirty(FIoFilenameHash Filenam
 {
 	if (!bRecacheAllFiles && FilenameHash != INVALID_IO_FILENAME_HASH)
 	{
+		LLM_SCOPE(ELLMTag::StreamingManager);
+		LLM_TAGSET_SCOPE_CLEAR(ELLMTagSet::Assets);
+		LLM_TAGSET_SCOPE_CLEAR(ELLMTagSet::AssetClasses);
+		UE_TRACE_METADATA_CLEAR_SCOPE();
 		FScopeLock ScopeLock(&MountedStateDirtyFilesCS);
 		MountedStateDirtyFiles.Emplace(FilenameHash);
 	}
@@ -363,6 +369,10 @@ void FRenderAssetStreamingManager::SetDisregardWorldResourcesForFrames( int32 Nu
 bool FRenderAssetStreamingManager::StreamOutRenderAssetData( int64 RequiredMemorySize )
 {
 	FScopeLock ScopeLock(&CriticalSection);
+	LLM_SCOPE(ELLMTag::StreamingManager);
+	LLM_TAGSET_SCOPE_CLEAR(ELLMTagSet::Assets);
+	LLM_TAGSET_SCOPE_CLEAR(ELLMTagSet::AssetClasses);
+	UE_TRACE_METADATA_CLEAR_SCOPE();
 
 	const int64 MaxTempMemoryAllowed = static_cast<int64>(Settings.MaxTempMemoryAllowed) * 1024 * 1024;
 	const bool CachedPauseTextureStreaming = bPauseRenderAssetStreaming;
@@ -832,6 +842,10 @@ void FRenderAssetStreamingManager::ProcessPendingLevelManagers()
 void FRenderAssetStreamingManager::AddLevel( ULevel* Level )
 {
 	FScopeLock ScopeLock(&CriticalSection);
+	LLM_SCOPE(ELLMTag::StreamingManager);
+	LLM_TAGSET_SCOPE_CLEAR(ELLMTagSet::Assets);
+	LLM_TAGSET_SCOPE_CLEAR(ELLMTagSet::AssetClasses);
+	UE_TRACE_METADATA_CLEAR_SCOPE();
 
 	check(Level);
 
@@ -942,6 +956,10 @@ void FRenderAssetStreamingManager::NotifyLevelOffset(ULevel* Level, const FVecto
 void FRenderAssetStreamingManager::AddStreamingRenderAsset(UStreamableRenderAsset* InAsset)
 {
 	FScopeLock ScopeLock(&CriticalSection);
+	LLM_SCOPE(ELLMTag::StreamingManager);
+	LLM_TAGSET_SCOPE_CLEAR(ELLMTagSet::Assets);
+	LLM_TAGSET_SCOPE_CLEAR(ELLMTagSet::AssetClasses);
+	UE_TRACE_METADATA_CLEAR_SCOPE();
 
 	STAT(GatheredStats.CallbacksCycles = -(int32)FPlatformTime::Cycles();)
 
@@ -964,6 +982,10 @@ void FRenderAssetStreamingManager::AddStreamingRenderAsset(UStreamableRenderAsse
 void FRenderAssetStreamingManager::RemoveStreamingRenderAsset( UStreamableRenderAsset* RenderAsset )
 {
 	FScopeLock ScopeLock(&CriticalSection);
+	LLM_SCOPE(ELLMTag::StreamingManager);
+	LLM_TAGSET_SCOPE_CLEAR(ELLMTagSet::Assets);
+	LLM_TAGSET_SCOPE_CLEAR(ELLMTagSet::AssetClasses);
+	UE_TRACE_METADATA_CLEAR_SCOPE();
 	TArray<FStreamingRenderAsset>& StreamingRenderAssets = GetStreamingRenderAssetsAsyncSafe();
 
 	STAT(GatheredStats.CallbacksCycles = -(int32)FPlatformTime::Cycles();)
diff --git a/Engine/Source/Runtime/OpenGLDrv/OpenGLDrv.Build.cs b/Engine/Source/Runtime/OpenGLDrv/OpenGLDrv.Build.cs
--- a/Engine/Source/Runtime/OpenGLDrv/OpenGLDrv.Build.cs
+++ b/Engine/Source/Runtime/OpenGLDrv/OpenGLDrv.Build.cs
@@ -19,7 +19,8 @@ public class OpenGLDrv : ModuleRules
 			"Engine",
 			"RHICore",
 			"RenderCore",
-			"PreLoadScreen"
+			"PreLoadScreen",
+			"TraceLog",
 		});
 
 		PrivateIncludePathModuleNames.Add("ImageWrapper");
diff --git a/Engine/Source/Runtime/OpenGLDrv/Private/OpenGLBuffer.cpp b/Engine/Source/Runtime/OpenGLDrv/Private/OpenGLBuffer.cpp
--- a/Engine/Source/Runtime/OpenGLDrv/Private/OpenGLBuffer.cpp
+++ b/Engine/Source/Runtime/OpenGLDrv/Private/OpenGLBuffer.cpp
@@ -7,6 +7,8 @@
 #include "CoreMinimal.h"
 #include "Containers/ResourceArray.h"
 #include "HAL/IConsoleManager.h"
+#include "HAL/LowLevelMemStats.h"
+#include "ProfilingDebugging/AssetMetadataTrace.h"
 #include "OpenGLDrv.h"
 #include "RHICoreBufferInitializer.h"
 
@@ -210,6 +212,10 @@ static GLenum GetNewBufferType(const FRHIBufferCreateDesc& CreateDesc)
 
 FRHIBufferInitializer FOpenGLDynamicRHI::RHICreateBufferInitializer(FRHICommandListBase& RHICmdList, const FRHIBufferCreateDesc& CreateDesc)
 {
+	LLM_SCOPE_DYNAMIC_STAT_OBJECTPATH_FNAME(CreateDesc.OwnerName, ELLMTagSet::Assets);
+	LLM_SCOPE_DYNAMIC_STAT_OBJECTPATH_FNAME(CreateDesc.GetTraceClassName(), ELLMTagSet::AssetClasses);
+	UE_TRACE_METADATA_SCOPE_ASSET_FNAME(CreateDesc.DebugName, CreateDesc.GetTraceClassName(), CreateDesc.OwnerName);
+
 	const GLenum BufferType = GetNewBufferType(CreateDesc);
 
 	const void* InitialData = nullptr;
diff --git a/Engine/Source/Runtime/OpenGLDrv/Private/OpenGLTexture.cpp b/Engine/Source/Runtime/OpenGLDrv/Private/OpenGLTexture.cpp
--- a/Engine/Source/Runtime/OpenGLDrv/Private/OpenGLTexture.cpp
+++ b/Engine/Source/Runtime/OpenGLDrv/Private/OpenGLTexture.cpp
@@ -11,7 +11,9 @@
 #include "RenderUtils.h"
 #include "OpenGLDrv.h"
 #include "OpenGLDrvPrivate.h"
+#include "HAL/LowLevelMemStats.h"
 #include "HAL/LowLevelMemTracker.h"
+#include "ProfilingDebugging/AssetMetadataTrace.h"
 #include "Engine/Texture.h"
 #include "RHICoreStats.h"
 
@@ -1400,6 +1402,9 @@ void FOpenGLTexture::CloneViaCopyImage(FOpenGLTexture* Src, uint32 InNumMips, in
 
 FTextureRHIRef FOpenGLDynamicRHI::RHICreateTexture(FRHICommandListBase& RHICmdList, const FRHITextureCreateDesc& CreateDesc)
 {
+	LLM_SCOPE_DYNAMIC_STAT_OBJECTPATH_FNAME(CreateDesc.OwnerName, ELLMTagSet::Assets);
+	LLM_SCOPE_DYNAMIC_STAT_OBJECTPATH_FNAME(CreateDesc.GetTraceClassName(), ELLMTagSet::AssetClasses);
+	UE_TRACE_METADATA_SCOPE_ASSET_FNAME(CreateDesc.DebugName, CreateDesc.GetTraceClassName(), CreateDesc.OwnerName);
 	return new FOpenGLTexture(RHICmdList, CreateDesc);
 }
 
@@ -1465,6 +1470,11 @@ FTextureRHIRef FOpenGLDynamicRHI::RHIAsyncReallocateTexture2D(FRHITexture* Textu
 		RHIGetDefaultResourceState(Desc.Flags, false),
 		TEXT("RHIAsyncReallocateTexture2D")
 	);
+	CreateDesc.SetOwnerName(OldTexture->GetOwnerName());
+
+	LLM_SCOPE_DYNAMIC_STAT_OBJECTPATH_FNAME(CreateDesc.OwnerName, ELLMTagSet::Assets);
+	LLM_SCOPE_DYNAMIC_STAT_OBJECTPATH_FNAME(CreateDesc.GetTraceClassName(), ELLMTagSet::AssetClasses);
+	UE_TRACE_METADATA_SCOPE_ASSET_FNAME(CreateDesc.DebugName, CreateDesc.GetTraceClassName(), CreateDesc.OwnerName);
 
 	FOpenGLTexture* NewTexture = new FOpenGLTexture(RHICmdList, CreateDesc);
 
diff --git a/Engine/Source/Runtime/RHI/Private/TextureProfiler.cpp b/Engine/Source/Runtime/RHI/Private/TextureProfiler.cpp
--- a/Engine/Source/Runtime/RHI/Private/TextureProfiler.cpp
+++ b/Engine/Source/Runtime/RHI/Private/TextureProfiler.cpp
@@ -5,6 +5,7 @@
 #include "Misc/CoreDelegates.h"
 #include "HAL/IConsoleManager.h"
 #include "ProfilingDebugging/CsvProfiler.h"
+#include "ProfilingDebugging/MetadataTrace.h"
 #include "RHIDefinitions.h"
 #include "RHIResources.h"
 
@@ -139,6 +140,7 @@ FTextureProfiler::FTextureDetails& FTextureProfiler::FTextureDetails::operator-=
 
 void FTextureProfiler::Init()
 {
+	LLM_SCOPE_BYNAME("Debug/TextureProfiler");
 	FCoreDelegates::OnEndFrameRT.AddRaw(this, &FTextureProfiler::Update);
 
 #if CSVPROFILERTRACE_ENABLED
@@ -165,6 +167,10 @@ FTextureProfiler* FTextureProfiler::Get()
 
 void FTextureProfiler::AddTextureAllocation(FRHITexture* UniqueTexturePtr, size_t Size, uint32 Alignment, size_t AllocationWaste)
 {
+	LLM_SCOPE_BYNAME("Debug/TextureProfiler");
+	LLM_TAGSET_SCOPE_CLEAR(ELLMTagSet::Assets);
+	LLM_TAGSET_SCOPE_CLEAR(ELLMTagSet::AssetClasses);
+	UE_TRACE_METADATA_CLEAR_SCOPE();
 	FScopeLock Lock(&TextureMapCS);
 
 	FTextureDetails AddedDetails(
@@ -207,6 +213,10 @@ void FTextureProfiler::AddTextureAllocation(FRHITexture* UniqueTexturePtr, size_
 
 void FTextureProfiler::UpdateTextureAllocation(FRHITexture* UniqueTexturePtr, size_t Size, uint32 Alignment, size_t AllocationWaste)
 {
+	LLM_SCOPE_BYNAME("Debug/TextureProfiler");
+	LLM_TAGSET_SCOPE_CLEAR(ELLMTagSet::Assets);
+	LLM_TAGSET_SCOPE_CLEAR(ELLMTagSet::AssetClasses);
+	UE_TRACE_METADATA_CLEAR_SCOPE();
 	FScopeLock Lock(&TextureMapCS);
 
 	FTextureDetails UpdatedDetails(
@@ -232,6 +242,10 @@ void FTextureProfiler::UpdateTextureAllocation(FRHITexture* UniqueTexturePtr, si
 
 void FTextureProfiler::RemoveTextureAllocation(FRHITexture* UniqueTexturePtr)
 {
+	LLM_SCOPE_BYNAME("Debug/TextureProfiler");
+	LLM_TAGSET_SCOPE_CLEAR(ELLMTagSet::Assets);
+	LLM_TAGSET_SCOPE_CLEAR(ELLMTagSet::AssetClasses);
+	UE_TRACE_METADATA_CLEAR_SCOPE();
 	FScopeLock Lock(&TextureMapCS);
 	FTextureDetails Details = TexturesMap.FindAndRemoveChecked(UniqueTexturePtr);
 	
@@ -244,6 +258,10 @@ void FTextureProfiler::RemoveTextureAllocation(FRHITexture* UniqueTexturePtr)
 
 void FTextureProfiler::UpdateTextureName(FRHITexture* UniqueTexturePtr)
 {
+	LLM_SCOPE_BYNAME("Debug/TextureProfiler");
+	LLM_TAGSET_SCOPE_CLEAR(ELLMTagSet::Assets);
+	LLM_TAGSET_SCOPE_CLEAR(ELLMTagSet::AssetClasses);
+	UE_TRACE_METADATA_CLEAR_SCOPE();
 	FScopeLock Lock(&TextureMapCS);
 
 	FTextureDetails* DetailsPtr = TexturesMap.Find(UniqueTexturePtr);
diff --git a/Engine/Source/Runtime/VulkanRHI/Private/VulkanBuffer.cpp b/Engine/Source/Runtime/VulkanRHI/Private/VulkanBuffer.cpp
--- a/Engine/Source/Runtime/VulkanRHI/Private/VulkanBuffer.cpp
+++ b/Engine/Source/Runtime/VulkanRHI/Private/VulkanBuffer.cpp
@@ -13,6 +13,8 @@
 #include "VulkanTransientResourceAllocator.h"
 #include "RHICoreBufferInitializer.h"
 #include "RHICoreStats.h"
+#include "HAL/LowLevelMemStats.h"
+#include "ProfilingDebugging/AssetMetadataTrace.h"
 
 struct FVulkanPendingBufferLock
 {
@@ -423,6 +425,9 @@ FRHIBufferInitializer FVulkanDynamicRHI::RHICreateBufferInitializer(FRHICommandL
 #else
 	LLM_SCOPE(EnumHasAnyFlags(CreateDesc.Usage, EBufferUsageFlags::VertexBuffer | EBufferUsageFlags::IndexBuffer) ? ELLMTag::Meshes : ELLMTag::RHIMisc);
 #endif
+	LLM_SCOPE_DYNAMIC_STAT_OBJECTPATH_FNAME(CreateDesc.OwnerName, ELLMTagSet::Assets);
+	LLM_SCOPE_DYNAMIC_STAT_OBJECTPATH_FNAME(CreateDesc.GetTraceClassName(), ELLMTagSet::AssetClasses);
+	UE_TRACE_METADATA_SCOPE_ASSET_FNAME(CreateDesc.DebugName, CreateDesc.GetTraceClassName(), CreateDesc.OwnerName);
 
 	FVulkanBuffer* Buffer = new FVulkanBuffer(Device, CreateDesc);
 
diff --git a/Engine/Source/Runtime/VulkanRHI/Private/VulkanTexture.cpp b/Engine/Source/Runtime/VulkanRHI/Private/VulkanTexture.cpp
--- a/Engine/Source/Runtime/VulkanRHI/Private/VulkanTexture.cpp
+++ b/Engine/Source/Runtime/VulkanRHI/Private/VulkanTexture.cpp
@@ -14,6 +14,8 @@
 #include "VulkanTransientResourceAllocator.h"
 #include "RHICoreStats.h"
 #include "RHILockTracker.h"
+#include "HAL/LowLevelMemStats.h"
+#include "ProfilingDebugging/AssetMetadataTrace.h"
 
 
 // This is a workaround for issues with AFBC on Mali GPUs before the G710
@@ -979,6 +981,9 @@ void FVulkanDynamicRHI::RHIUpdateTextureReference(FRHICommandListBase& RHICmdLis
 FVulkanDynamicRHI::FCreateTextureResult FVulkanDynamicRHI::BeginCreateTextureInternal(const FRHITextureCreateDesc& CreateDesc, const FRHITransientHeapAllocation* InTransientHeapAllocation)
 {
 	LLM_SCOPE_VULKAN(GetMemoryTagForTextureFlags(CreateDesc.Flags));
+	LLM_SCOPE_DYNAMIC_STAT_OBJECTPATH_FNAME(CreateDesc.OwnerName, ELLMTagSet::Assets);
+	LLM_SCOPE_DYNAMIC_STAT_OBJECTPATH_FNAME(CreateDesc.GetTraceClassName(), ELLMTagSet::AssetClasses);
+	UE_TRACE_METADATA_SCOPE_ASSET_FNAME(CreateDesc.DebugName, CreateDesc.GetTraceClassName(), CreateDesc.OwnerName);
 
 	FVulkanTexture* Texture = new FVulkanTexture(*Device, CreateDesc, InTransientHeapAllocation);
 
@@ -1166,7 +1171,8 @@ FTextureRHIRef FVulkanDynamicRHI::AsyncReallocateTexture2D_RenderThread(FRHIComm
 		.SetFlags(OldDesc.Flags)
 		.SetNumMips(NewMipCount)
 		.SetNumSamples(OldDesc.NumSamples)
-		.DetermineInititialState();
+		.DetermineInititialState()
+		.SetOwnerName(OldTexture->GetOwnerName());
 
 	FVulkanTexture* NewTexture = CreateTextureInternal(RHICmdList, CreateDesc, {});
 
diff --git a/Engine/Source/Runtime/VulkanRHI/VulkanRHI.Build.cs b/Engine/Source/Runtime/VulkanRHI/VulkanRHI.Build.cs
--- a/Engine/Source/Runtime/VulkanRHI/VulkanRHI.Build.cs
+++ b/Engine/Source/Runtime/VulkanRHI/VulkanRHI.Build.cs
@@ -38,10 +38,11 @@ public class VulkanRHI : ModuleRules
 				"RHICore",
 				"RenderCore", 
 				"HeadMountedDisplay",
-                "PreLoadScreen",
-				"BuildSettings"
-            }
-        );
+				"PreLoadScreen",
+				"BuildSettings",
+				"TraceLog",
+			}
+		);
 
 		if (Target.Platform == UnrealTargetPlatform.Android)
 		{
diff --git a/Engine/Source/Runtime/Windows/D3D11RHI/D3D11RHI.Build.cs b/Engine/Source/Runtime/Windows/D3D11RHI/D3D11RHI.Build.cs
--- a/Engine/Source/Runtime/Windows/D3D11RHI/D3D11RHI.Build.cs
+++ b/Engine/Source/Runtime/Windows/D3D11RHI/D3D11RHI.Build.cs
@@ -14,7 +14,8 @@ public class D3D11RHI : ModuleRules
 			"CoreUObject",
 			"Engine",
 			"RHICore",
-			"RenderCore"
+			"RenderCore",
+			"TraceLog"
 		);
 
 		PublicDependencyModuleNames.AddAll(
diff --git a/Engine/Source/Runtime/Windows/D3D11RHI/Private/D3D11Buffer.cpp b/Engine/Source/Runtime/Windows/D3D11RHI/Private/D3D11Buffer.cpp
--- a/Engine/Source/Runtime/Windows/D3D11RHI/Private/D3D11Buffer.cpp
+++ b/Engine/Source/Runtime/Windows/D3D11RHI/Private/D3D11Buffer.cpp
@@ -6,6 +6,8 @@
 
 #include "D3D11RHIPrivate.h"
 #include "RHICoreBufferInitializer.h"
+#include "HAL/LowLevelMemStats.h"
+#include "ProfilingDebugging/AssetMetadataTrace.h"
 
 TAutoConsoleVariable<int32> GCVarUseSharedKeyedMutex(
 	TEXT("r.D3D11.UseSharedKeyMutex"),
@@ -154,6 +156,10 @@ FD3D11Buffer* FD3D11DynamicRHI::CreateBufferInternal(const FRHIBufferCreateDesc&
 
 FRHIBufferInitializer FD3D11DynamicRHI::RHICreateBufferInitializer(FRHICommandListBase& RHICmdList, const FRHIBufferCreateDesc& CreateDesc)
 {
+	LLM_SCOPE_DYNAMIC_STAT_OBJECTPATH_FNAME(CreateDesc.OwnerName, ELLMTagSet::Assets);
+	LLM_SCOPE_DYNAMIC_STAT_OBJECTPATH_FNAME(CreateDesc.GetTraceClassName(), ELLMTagSet::AssetClasses);
+	UE_TRACE_METADATA_SCOPE_ASSET_FNAME(CreateDesc.DebugName, CreateDesc.GetTraceClassName(), CreateDesc.OwnerName);
+
 	if (CreateDesc.IsNull())
 	{
 		FD3D11Buffer* Buffer = BeginCreateBufferInternal(CreateDesc);
diff --git a/Engine/Source/Runtime/Windows/D3D11RHI/Private/D3D11Texture.cpp b/Engine/Source/Runtime/Windows/D3D11RHI/Private/D3D11Texture.cpp
--- a/Engine/Source/Runtime/Windows/D3D11RHI/Private/D3D11Texture.cpp
+++ b/Engine/Source/Runtime/Windows/D3D11RHI/Private/D3D11Texture.cpp
@@ -16,8 +16,10 @@
 #endif
 #include "Windows/HideWindowsPlatformTypes.h"
 
+#include "HAL/LowLevelMemStats.h"
 #include "HAL/LowLevelMemTracker.h"
 #include "ProfilingDebugging/MemoryTrace.h"
+#include "ProfilingDebugging/AssetMetadataTrace.h"
 #include "RHICoreStats.h"
 #include "RHIUtilities.h"
 #include "RHITextureUtils.h"
@@ -807,6 +809,9 @@ FD3D11Texture* FD3D11DynamicRHI::FinalizeCreateTextureInternal(FD3D11Texture* Te
 
 FD3D11Texture* FD3D11DynamicRHI::CreateTextureInternal(const FRHITextureCreateDesc& CreateDesc, TConstArrayView<uint8> InitialData)
 {
+	LLM_SCOPE_DYNAMIC_STAT_OBJECTPATH_FNAME(CreateDesc.OwnerName, ELLMTagSet::Assets);
+	LLM_SCOPE_DYNAMIC_STAT_OBJECTPATH_FNAME(CreateDesc.GetTraceClassName(), ELLMTagSet::AssetClasses);
+	UE_TRACE_METADATA_SCOPE_ASSET_FNAME(CreateDesc.DebugName, CreateDesc.GetTraceClassName(), CreateDesc.OwnerName);
 	FD3D11Texture* Texture = BeginCreateTextureInternal(CreateDesc);
 	return FinalizeCreateTextureInternal(Texture, CreateDesc, InitialData);
 }
@@ -940,6 +945,7 @@ FTextureRHIRef FD3D11DynamicRHI::RHIAsyncReallocateTexture2D(FRHITexture* Textur
 	);
 	CreateDesc.Extent = FIntPoint(NewSizeX, NewSizeY);
 	CreateDesc.NumMips = NewMipCount;
+	CreateDesc.SetOwnerName(Texture2D->GetOwnerName());
 
 	// Allocate a new texture.
 	FD3D11Texture* NewTexture2D = CreateTextureInternal(CreateDesc, {});
@@ -973,6 +979,7 @@ FTextureRHIRef FD3D11DynamicRHI::AsyncReallocateTexture2D_RenderThread(
 		);
 		CreateDesc.Extent = FIntPoint(NewSizeX, NewSizeY);
 		CreateDesc.NumMips = NewMipCount;
+		CreateDesc.SetOwnerName(Texture2D->GetOwnerName());
 
 		NewTexture2D = CreateTextureInternal(CreateDesc, {});
 
