Fix non-unity

diff --git a/Engine/Source/Runtime/D3D12RHI/Private/D3D12Viewport.h b/Engine/Source/Runtime/D3D12RHI/Private/D3D12Viewport.h
--- a/Engine/Source/Runtime/D3D12RHI/Private/D3D12Viewport.h
+++ b/Engine/Source/Runtime/D3D12RHI/Private/D3D12Viewport.h
@@ -7,6 +7,7 @@
 #pragma once
 
 #include "D3D12RHICommon.h"
+#include "D3D12Texture.h"
 #include "HAL/Runnable.h"
 #include "MultiGPU.h"
 #include "RHIResources.h"
