Don't use TRHICommandList_RecursiveHazardous for copy upload memory to inline SBT raytracing buffers but use provided context directly (fix issue with mGPU setup)

diff --git a/Engine/Source/Runtime/D3D12RHI/Private/D3D12RayTracing.cpp b/Engine/Source/Runtime/D3D12RHI/Private/D3D12RayTracing.cpp
--- a/Engine/Source/Runtime/D3D12RHI/Private/D3D12RayTracing.cpp
+++ b/Engine/Source/Runtime/D3D12RHI/Private/D3D12RayTracing.cpp
@@ -1715,14 +1715,26 @@ public:
 		{
 			check(InInlineBindingDataBuffer);
 
-			uint32 MaxUsedInlineRecordIndex = WorkerData[0].MaxUsedHitRecordIndex / NumShaderSlotsPerGeometrySegment;
+			FD3D12Buffer* D3D12InlineBindingDataBuffer = FD3D12DynamicRHI::ResourceCast(InInlineBindingDataBuffer, Device->GetGPUIndex());
 
-			TRHICommandList_RecursiveHazardous<FD3D12CommandContext> RHICmdList(&Context);
-			RHICmdList.SetAllowExtraTransition(false);
+			uint32 MaxUsedInlineRecordIndex = WorkerData[0].MaxUsedHitRecordIndex / NumShaderSlotsPerGeometrySegment;
 			uint32 MaxUsedInlineGeometryParamaterDataSize = (MaxUsedInlineRecordIndex + 1) * sizeof(FD3D12HitGroupSystemParameters);
-			void* MappedUploadMemory = GDynamicRHI->RHILockBuffer(RHICmdList, InInlineBindingDataBuffer, 0, MaxUsedInlineGeometryParamaterDataSize, EResourceLockMode::RLM_WriteOnly);
+
+			FD3D12ResourceLocation UploadResourceLocation(Device);
+			void* MappedUploadMemory = Adapter->GetUploadHeapAllocator(Device->GetGPUIndex()).AllocUploadResource(MaxUsedInlineGeometryParamaterDataSize, D3D12_TEXTURE_DATA_PLACEMENT_ALIGNMENT, UploadResourceLocation);
 			FMemory::Memcpy(MappedUploadMemory, InlineGeometryParameterData.GetData(), MaxUsedInlineGeometryParamaterDataSize);
-			GDynamicRHI->RHIUnlockBuffer(RHICmdList, InInlineBindingDataBuffer);
+
+			FD3D12Resource* SourceResource = UploadResourceLocation.GetResource();
+			FD3D12Resource* DestResource = D3D12InlineBindingDataBuffer->ResourceLocation.GetResource();
+						
+			Context.UpdateResidency(DestResource);
+			Context.CopyBufferRegionChecked(
+				DestResource->GetResource(), DestResource->GetName(),
+				D3D12InlineBindingDataBuffer->ResourceLocation.GetOffsetFromBaseOfResource(),
+				SourceResource->GetResource(), SourceResource->GetName(),
+				UploadResourceLocation.GetOffsetFromBaseOfResource(),
+				MaxUsedInlineGeometryParamaterDataSize
+			);
 
 			#if ENABLE_RESIDENCY_MANAGEMENT
 			// assume no references resources when inline SBT only - inline SBT only has references to IB/VB buffers and these are 
