Add missing check for GEnableResidencyManagement to prevent using ResidencyManager when it might not be initialized

diff --git a/Engine/Source/Runtime/D3D12RHI/Private/D3D12Resources.cpp b/Engine/Source/Runtime/D3D12RHI/Private/D3D12Resources.cpp
--- a/Engine/Source/Runtime/D3D12RHI/Private/D3D12Resources.cpp
+++ b/Engine/Source/Runtime/D3D12RHI/Private/D3D12Resources.cpp
@@ -516,7 +516,7 @@ void FD3D12Resource::CommitReservedResource(ID3D12CommandQueue* D3DCommandQueue,
 #if ENABLE_RESIDENCY_MANAGEMENT
 	FD3D12ResidencyManager& ResidencyManager = GetParentDevice()->GetResidencyManager();
 
-	if (!UsedResidencyHandles.IsEmpty())
+	if (GEnableResidencyManagement && !UsedResidencyHandles.IsEmpty())
 	{
 		HRESULT HR = S_OK;
 
@@ -554,7 +554,7 @@ void FD3D12Resource::CommitReservedResource(ID3D12CommandQueue* D3DCommandQueue,
 	}
 
 #if ENABLE_RESIDENCY_MANAGEMENT
-	if (!UsedResidencyHandles.IsEmpty())
+	if (GEnableResidencyManagement && !UsedResidencyHandles.IsEmpty())
 	{
 		// Signal the fence for this queue after UpdateTileMappings complete.
 		// This is analogous to executing a command list that references a set of resources.
