Fix D3D12 allocation tracking CVar, allowing it to be set.

diff --git a/Engine/Source/Runtime/D3D12RHI/Private/D3D12Adapter.cpp b/Engine/Source/Runtime/D3D12RHI/Private/D3D12Adapter.cpp
--- a/Engine/Source/Runtime/D3D12RHI/Private/D3D12Adapter.cpp
+++ b/Engine/Source/Runtime/D3D12RHI/Private/D3D12Adapter.cpp
@@ -120,10 +120,10 @@ ED3D12DredMode D3D12_GetDredMode()
 	}
 }
 
-bool GD3D12TrackAllAlocations = false;
-static TAutoConsoleVariable<int32> CVarD3D12TrackAllAllocations(
+bool GD3D12TrackAllAllocations = false;
+static FAutoConsoleVariableRef CVarD3D12TrackAllAllocations(
 	TEXT("D3D12.TrackAllAllocations"),
-	GD3D12TrackAllAlocations,
+	GD3D12TrackAllAllocations,
 	TEXT("Controls whether D3D12 RHI should track all allocation information (default = off)."),
 	ECVF_ReadOnly
 );
@@ -1283,7 +1283,7 @@ void FD3D12Adapter::InitializeDevices()
 		// Set flag if we want to track all allocations - comes with some overhead and only possible when Tier 2 is available
 		// (because we will create placed buffers for texture allocation to retrieve the GPU virtual addresses)
 		const bool bTraceMemAlloc = UE_TRACE_CHANNELEXPR_IS_ENABLED(MemAllocChannel);
-		bTrackAllAllocation = (GD3D12TrackAllAlocations || UE::RHI::UseGPUCrashDebugging() || bTraceMemAlloc) && (GetResourceHeapTier() == D3D12_RESOURCE_HEAP_TIER_2);
+		bTrackAllAllocation = (GD3D12TrackAllAllocations || UE::RHI::UseGPUCrashDebugging() || bTraceMemAlloc) && (GetResourceHeapTier() == D3D12_RESOURCE_HEAP_TIER_2);
 #endif
 
 		ERHIBindlessConfiguration BindlessResourcesConfig = ERHIBindlessConfiguration::Disabled;
