Improved custom filtering of D3D12 Debug Layer Message. Now we can specify a sub string to match the errormessage  in addition to the error ID allowing a more selective filtering arounf known issue without hide real issue.

diff --git a/Engine/Source/Runtime/D3D12RHI/Private/D3D12Adapter.cpp b/Engine/Source/Runtime/D3D12RHI/Private/D3D12Adapter.cpp
--- a/Engine/Source/Runtime/D3D12RHI/Private/D3D12Adapter.cpp
+++ b/Engine/Source/Runtime/D3D12RHI/Private/D3D12Adapter.cpp
@@ -138,15 +138,54 @@ static FAutoConsoleVariableRef CVarD3D12EvictAllResidentResourcesInBackground(
 #endif // PLATFORM_WINDOWS
 
 #if D3D12_SUPPORTS_INFO_QUEUE
+D3D12_MESSAGE_SEVERITY GDebugLayerMinimumSeverity = D3D12_MESSAGE_SEVERITY_ERROR;
+
+struct MessageFilter
+{
+	MessageFilter(D3D12_MESSAGE_ID ID, const TCHAR* MessageSubString ): MessageID(ID), SubString(MessageSubString)
+	{}
+
+	D3D12_MESSAGE_ID MessageID;
+	FString SubString;
+};
+
+TArray<MessageFilter> GIgnoreD3D12MessageArray;
+
+// Returns true if a message is in the ignore list, or with not enough severity, and we don't want to show it
+static bool IsMessageFilteredOut(D3D12_MESSAGE_SEVERITY Severity, D3D12_MESSAGE_ID ID, LPCSTR pDescription)
+{
+	if (Severity > GDebugLayerMinimumSeverity)
+	{
+		return true;
+	}
+
+	for (const MessageFilter& MessageToFilter : GIgnoreD3D12MessageArray)
+	{
+		if (ID == MessageToFilter.MessageID)
+		{
+			if ((MessageToFilter.SubString.IsEmpty()) || (FCStringAnsi::Strstr(pDescription, TCHAR_TO_ANSI(*MessageToFilter.SubString)) != nullptr))
+			{
+				return true;
+			}
+		}
+	}
+
+	return false;
+}
 
 // Log with the UE_LOG the message received from D3D12 debug layer
 // Return if we need to break execution or not
-static bool LogD3D12Message(D3D12_MESSAGE_SEVERITY Severity, LPCSTR pDescription)
+static bool LogD3D12Message(D3D12_MESSAGE_SEVERITY Severity, D3D12_MESSAGE_ID ID, LPCSTR pDescription)
 {
 	bool bBreakOnLog = false;
 
 	if (pDescription)
 	{
+		if (IsMessageFilteredOut(Severity, ID, pDescription))
+		{
+			return false;
+		}
+
 		switch (Severity)
 		{
 		case D3D12_MESSAGE_SEVERITY_CORRUPTION:
@@ -219,7 +258,7 @@ static bool CheckD3DStoredMessages()
 					// get the actual message data from the queue
 					hr = d3dInfoQueue->GetMessage(MessageIndex, d3dMessage, &MessageLength);
 
-					bResult = LogD3D12Message(d3dMessage->Severity, d3dMessage->pDescription);
+					bResult = LogD3D12Message(d3dMessage->Severity, d3dMessage->ID, d3dMessage->pDescription);
 				}
 			}
 			d3dInfoQueue->ClearStoredMessages();
@@ -269,7 +308,7 @@ static LONG __stdcall D3DVectoredExceptionHandler(EXCEPTION_POINTERS* InInfo)
 #if D3D12_SUPPORTS_INFO_QUEUE
 void D3D12MessageCallBack(D3D12_MESSAGE_CATEGORY Category, D3D12_MESSAGE_SEVERITY Severity, D3D12_MESSAGE_ID ID, LPCSTR pDescription, void* pContext)
 {
-	bool bBreakOnLog = LogD3D12Message(Severity, pDescription);
+	bool bBreakOnLog = LogD3D12Message(Severity, ID, pDescription);
 	
 	if (FPlatformMisc::IsDebuggerPresent() && bBreakOnLog)
 	{
@@ -813,13 +852,22 @@ void FD3D12Adapter::CreateD3DInfoQueue()
 		}
 #endif // PLATFORM_DESKTOP
 
+		GIgnoreD3D12MessageArray.Add(MessageFilter(D3D12_MESSAGE_ID_BUILD_RAYTRACING_ACCELERATION_STRUCTURE_INVALID, TEXT("exceeds end of the virtual address range of Heap")));
+
 		if (GSupportsEfficientAsyncCompute)
 		{
 			// This is not optimal to have it disable but this should be absolutely removed when DX12 can use EB 
-			// The reason why this is here is because when async compute is on we can receive from RDG a transition UAV -> SRVMask on the Async Pipe
-			// DX12 cannot transition to Gfx states on the Async Pipe therefore we do: UAV->SRVCompute on Async, Fence Async -> Gfx, SRVCompute->SRVMask on Gfx
-			// Since the debug layer is on the CPU and the submission is happening in parallel threads the validation layer cannot known that the order on the GPU will be UAV->SRVCompute->SRVMask
-			DenyIds.Add(D3D12_MESSAGE_ID_RESOURCE_BARRIER_BEFORE_AFTER_MISMATCH);
+ 			// The reason why this is here is because when async compute is on we can receive from RDG a transition UAV -> SRVMask on the Async Pipe
+ 			// DX12 cannot transition to Gfx states on the Async Pipe therefore we do: UAV->SRVCompute on Async, Fence Async -> Gfx, SRVCompute->SRVMask on Gfx
+ 			// Since the debug layer is on the CPU and the submission is happening in parallel threads the validation layer cannot known that the order on the GPU will be UAV->SRVCompute->SRVMask
+			GIgnoreD3D12MessageArray.Add(MessageFilter(D3D12_MESSAGE_ID_RESOURCE_BARRIER_BEFORE_AFTER_MISMATCH, TEXT("does not match with the state (0x40: D3D12_RESOURCE_STATE_NON_PIXEL_SHADER_RESOURCE) specified in preceding")));
+			GIgnoreD3D12MessageArray.Add(MessageFilter(D3D12_MESSAGE_ID_RESOURCE_BARRIER_BEFORE_AFTER_MISMATCH, TEXT("does not match with the state (0x8: D3D12_RESOURCE_STATE_UNORDERED_ACCESS) specified in preceding")));
+		}
+
+		// Adding all DenyIds as well with no sub string
+		for (int i = 0; i < DenyIds.Num(); i++)
+		{
+			GIgnoreD3D12MessageArray.Add(MessageFilter(DenyIds[i], TEXT("")));
 		}
 
 		NewFilter.DenyList.NumIDs = DenyIds.Num();
@@ -846,7 +894,13 @@ void FD3D12Adapter::CreateD3DInfoQueue()
 			if (pd3dInfoQueue1)
 			{
 				void* ParamContext = this;
-				pd3dInfoQueue1->RegisterMessageCallback(D3D12MessageCallBack, D3D12_MESSAGE_CALLBACK_FLAG_NONE, ParamContext, &GD3D12CallBackCookie);
+				
+				// To control exactly what we want to print out we need to disale the Debug output and ignore all filters if any
+				pd3dInfoQueue1->RegisterMessageCallback(D3D12MessageCallBack, D3D12_MESSAGE_CALLBACK_IGNORE_FILTERS, ParamContext, &GD3D12CallBackCookie);
+				pd3dInfoQueue1->SetMuteDebugOutput(true);
+
+				// Set the minimum log severity
+				GDebugLayerMinimumSeverity = (bLogWarnings) ? D3D12_MESSAGE_SEVERITY_WARNING : D3D12_MESSAGE_SEVERITY_ERROR;
 
 				if (GD3D12CallBackCookie == 0)
 				{
