Remove hardcoded resource name and add Queue type information to DX12 low level logging

diff --git a/Engine/Source/Runtime/D3D12RHI/Private/D3D12CommandList.cpp b/Engine/Source/Runtime/D3D12RHI/Private/D3D12CommandList.cpp
--- a/Engine/Source/Runtime/D3D12RHI/Private/D3D12CommandList.cpp
+++ b/Engine/Source/Runtime/D3D12RHI/Private/D3D12CommandList.cpp
@@ -270,9 +270,8 @@ void FD3D12CommandList::Close()
 	}
 
 #if DEBUG_RESOURCE_STATES
-	#define DX12_RESOURCE_NAME_TO_LOG	"Lumen.RadianceCache.DepthProbeAtlasTexture"
 	if (hr != S_OK)
-		LogResourceBarriers(State.ResourceBarriers, Interfaces.CommandList.GetReference(), FString(DX12_RESOURCE_NAME_TO_LOG));
+		LogResourceBarriers(State.ResourceBarriers, Interfaces.CommandList.GetReference(), ED3D12QueueType::Direct , FString(DX12_RESOURCE_NAME_TO_LOG));
 #endif
 
 	VERIFYD3D12RESULT(hr);
diff --git a/Engine/Source/Runtime/D3D12RHI/Private/D3D12RHIDefinitions.h b/Engine/Source/Runtime/D3D12RHI/Private/D3D12RHIDefinitions.h
--- a/Engine/Source/Runtime/D3D12RHI/Private/D3D12RHIDefinitions.h
+++ b/Engine/Source/Runtime/D3D12RHI/Private/D3D12RHIDefinitions.h
@@ -15,6 +15,7 @@
 #define RESIDENCY_PIPELINE_DEPTH	6
 
 #define DEBUG_RESOURCE_STATES	0
+#define DX12_RESOURCE_NAME_TO_LOG	""
 
 #define D3D12_RHI_RAYTRACING (RHI_RAYTRACING)
 
diff --git a/Engine/Source/Runtime/D3D12RHI/Private/D3D12Resources.cpp b/Engine/Source/Runtime/D3D12RHI/Private/D3D12Resources.cpp
--- a/Engine/Source/Runtime/D3D12RHI/Private/D3D12Resources.cpp
+++ b/Engine/Source/Runtime/D3D12RHI/Private/D3D12Resources.cpp
@@ -1714,10 +1714,9 @@ void FD3D12ResourceBarrierBatcher::FlushIntoCommandList(FD3D12CommandList& Comma
 		}
 
 #if DEBUG_RESOURCE_STATES
-		#define DX12_RESOURCE_NAME_TO_LOG	"Test"
    		TArrayView<FD3D12ResourceBarrier> SubsetBarriers(Barriers.GetData() + BatchStart, BatchEnd - BatchStart);
    		TConstArrayView<D3D12_RESOURCE_BARRIER> ConstArrayView(	reinterpret_cast<const D3D12_RESOURCE_BARRIER*>(SubsetBarriers.GetData()), SubsetBarriers.Num());
- 		LogResourceBarriers(ConstArrayView, CommandList.Interfaces.CommandList.GetReference(), FString(DX12_RESOURCE_NAME_TO_LOG));
+		LogResourceBarriers(ConstArrayView, CommandList.Interfaces.CommandList.GetReference(), CommandList.QueueType, FString(DX12_RESOURCE_NAME_TO_LOG));
 #endif
 		CommandList.GraphicsCommandList()->ResourceBarrier(BatchEnd - BatchStart, &Barriers[BatchStart]);
 #if DEBUG_RESOURCE_STATES
diff --git a/Engine/Source/Runtime/D3D12RHI/Private/D3D12Util.cpp b/Engine/Source/Runtime/D3D12RHI/Private/D3D12Util.cpp
--- a/Engine/Source/Runtime/D3D12RHI/Private/D3D12Util.cpp
+++ b/Engine/Source/Runtime/D3D12RHI/Private/D3D12Util.cpp
@@ -93,10 +93,12 @@ FString GetD312ObjectName(ID3D12Object* const Object)
 	if (Object == nullptr)
 		return FString("Unknown Resource");
 
-	TCHAR OutName[128];
-	UINT size = 128;
+	#define MAX_OBJECT_NAME_LEN 512
+
+	TCHAR OutName[MAX_OBJECT_NAME_LEN];
+	UINT size = MAX_OBJECT_NAME_LEN;
 	HRESULT hr = Object->GetPrivateData(WKPDID_D3DDebugObjectNameW, &size, OutName);
-	check(size <= 128);
+	check(size <= MAX_OBJECT_NAME_LEN);
 
 	if (hr != S_OK)
 		return FString("Unknown Resource");
@@ -1586,7 +1588,7 @@ FString ConvertToResourceStateString(uint32 ResourceState)
 	return ResourceStateString;
 }
 
-void LogResourceBarriers(TConstArrayView<D3D12_RESOURCE_BARRIER> Barriers, ID3D12CommandList* const pCommandList, const FString& ResourceName)
+void LogResourceBarriers(TConstArrayView<D3D12_RESOURCE_BARRIER> Barriers, ID3D12CommandList* const pCommandList, ED3D12QueueType QueueType, const FString& ResourceName)
 {
 	// Configure what resource barriers are logged.
 	const bool bLogAll = true;
@@ -1629,7 +1631,7 @@ void LogResourceBarriers(TConstArrayView<D3D12_RESOURCE_BARRIER> Barriers, ID3D1
 
 				if ((bCheckResourceName == false) || (BarrierResourceName == ResourceName))
 				{
-					UE_LOG(LogD3D12RHI, Log, TEXT("*** BARRIER (CmdList: %016llX) %u/%u: %s %016llX (Sub: %u), %s -> %s"), pCommandList, i + 1, Barriers.Num(),
+					UE_LOG(LogD3D12RHI, Log, TEXT("*** BARRIER (CmdList: %016llX, Queue: %s) %u/%u: %s %016llX (Sub: %u), %s -> %s"), pCommandList, GetD3DCommandQueueTypeName(QueueType), i + 1, Barriers.Num(),
 						*BarrierResourceName,
 						currentBarrier.Transition.pResource,
 						currentBarrier.Transition.Subresource,
@@ -1645,7 +1647,7 @@ void LogResourceBarriers(TConstArrayView<D3D12_RESOURCE_BARRIER> Barriers, ID3D1
 				const FString BarrierResourceName = GetD312ObjectName(currentBarrier.UAV.pResource);
 				if ((bCheckResourceName == false) || (BarrierResourceName == ResourceName))
 				{
-					UE_LOG(LogD3D12RHI, Log, TEXT("*** BARRIER (CmdList: %016llX) %u/%u: UAV Barrier %s"), pCommandList, i + 1, Barriers.Num(), *BarrierResourceName);
+					UE_LOG(LogD3D12RHI, Log, TEXT("*** BARRIER (CmdList: %016llX, Queue: %s) %u/%u: UAV Barrier %s"), pCommandList, GetD3DCommandQueueTypeName(QueueType), i + 1, Barriers.Num(), *BarrierResourceName);
 				}
 			}
 			break;
@@ -1657,7 +1659,7 @@ void LogResourceBarriers(TConstArrayView<D3D12_RESOURCE_BARRIER> Barriers, ID3D1
 				
 				if ((bCheckResourceName == false) || (BarrierResourceNameBefore == ResourceName) || (BarrierResourceNameAfter == ResourceName))
 				{
-					UE_LOG(LogD3D12RHI, Log, TEXT("*** BARRIER (CmdList: %016llX) %u/%u: Aliasing Barrier, %016llX %s -> %016llX %s"), pCommandList, i + 1, Barriers.Num(), currentBarrier.Aliasing.pResourceBefore, *BarrierResourceNameBefore, currentBarrier.Aliasing.pResourceAfter, *BarrierResourceNameAfter);
+					UE_LOG(LogD3D12RHI, Log, TEXT("*** BARRIER (CmdList: %016llX, Queue: %s) %u/%u: Aliasing Barrier, %016llX %s -> %016llX %s"), pCommandList, GetD3DCommandQueueTypeName(QueueType), i + 1, Barriers.Num(), currentBarrier.Aliasing.pResourceBefore, *BarrierResourceNameBefore, currentBarrier.Aliasing.pResourceAfter, *BarrierResourceNameAfter);
 				}
 			}
 			break;
diff --git a/Engine/Source/Runtime/D3D12RHI/Private/D3D12Util.h b/Engine/Source/Runtime/D3D12RHI/Private/D3D12Util.h
--- a/Engine/Source/Runtime/D3D12RHI/Private/D3D12Util.h
+++ b/Engine/Source/Runtime/D3D12RHI/Private/D3D12Util.h
@@ -256,7 +256,7 @@ private:
 
 void LogExecuteCommandLists(uint32 NumCommandLists, ID3D12CommandList *const *ppCommandLists);
 FString ConvertToResourceStateString(uint32 ResourceState);
-void LogResourceBarriers(TConstArrayView<D3D12_RESOURCE_BARRIER> Barriers, ID3D12CommandList *const pCommandList, const FString& ResourceName);
+void LogResourceBarriers(TConstArrayView<D3D12_RESOURCE_BARRIER> Barriers, ID3D12CommandList *const pCommandList, ED3D12QueueType QueueType, const FString& ResourceName);
 
 // Custom resource states
 // To Be Determined (TBD) means we need to fill out a resource barrier before the command list is executed.
