FD3D12DynamicRHI::ReadSurfaceDataNoMSAARaw: Trace out memory being freed for a temporary staging texture if one was created for readback. Without this, Memory Insights reports ever-growing memory usage.

diff --git a/Engine/Source/Runtime/D3D12RHI/Private/D3D12RenderTarget.cpp b/Engine/Source/Runtime/D3D12RHI/Private/D3D12RenderTarget.cpp
--- a/Engine/Source/Runtime/D3D12RHI/Private/D3D12RenderTarget.cpp
+++ b/Engine/Source/Runtime/D3D12RHI/Private/D3D12RenderTarget.cpp
@@ -537,6 +537,19 @@ void FD3D12DynamicRHI::ReadSurfaceDataNoMSAARaw(FRHITexture* TextureRHI, FIntRec
 	}
 
 	TempTexture2D->GetResource()->Unmap(0, nullptr);
+
+#if UE_MEMORY_TRACE_ENABLED
+	// We only need to free memory for a texture that was created specifically for this operation -
+	// not for a texture that was reused for readback.
+	FD3D12Texture* Texture = GetD3D12TextureFromRHITexture(TextureRHI, GPUIndex);
+	if (TempTexture2D != Texture->GetResource())
+	{
+		check(TempTexture2D.GetRefCount() == 1);
+		// Free the temporary texture after read back finishes.
+		// This matches the MemoryTrace_Alloc call in Adapter->CreateBuffer() via the GetStagingTexture call above.
+		MemoryTrace_Free(TempTexture2D->GetGPUVirtualAddress(), EMemoryTraceRootHeap::VideoMemory);
+	}
+#endif
 }
 
 void FD3D12DynamicRHI::RHIReadSurfaceData(FRHITexture* TextureRHI, FIntRect InRect, TArray<FLinearColor>& OutData, FReadSurfaceDataFlags InFlags)
