Inside D3D12's ClearUAV with bindless enabled, make sure we are using an active bindless GPU heap if there was a resize from the temporary descriptor allocation.
FlushPendingDescriptorUpdates should queue up an invalidation of RS binding instead of requiring the caller to pass its return value to ApplyState.

diff --git a/Engine/Source/Runtime/D3D12RHI/Private/D3D12BindlessDescriptors.cpp b/Engine/Source/Runtime/D3D12RHI/Private/D3D12BindlessDescriptors.cpp
--- a/Engine/Source/Runtime/D3D12RHI/Private/D3D12BindlessDescriptors.cpp
+++ b/Engine/Source/Runtime/D3D12RHI/Private/D3D12BindlessDescriptors.cpp
@@ -482,7 +482,7 @@ void FD3D12BindlessResourceManager::UpdateDescriptor(FD3D12ContextArray const& C
 				if (ContextBase)
 				{
 					FD3D12CommandContext& Context = *ContextBase->GetSingleDeviceContext(GPUIndex);
-					Context.GetBindlessState().bRequestNewGpuHeap = true;
+					Context.GetBindlessState().RefreshDescriptorHeap();
 					check(!Context.GetExecutingCommandList().AllowParallelTranslate());
 				}
 			}
@@ -495,14 +495,12 @@ void FD3D12BindlessResourceManager::UpdateDescriptor(FD3D12ContextArray const& C
 	}
 }
 
-bool FD3D12BindlessResourceManager::FlushPendingDescriptorUpdates(FD3D12CommandContext& Context)
+void FD3D12BindlessResourceManager::FlushPendingDescriptorUpdates(FD3D12CommandContext& Context)
 {
-	bool bSetNewHeaps = false;
-
 	FD3D12ContextBindlessState& State = Context.GetBindlessState();
 
 	// Create a new heap because there have been descriptor updates?
-	if (State.bRequestNewGpuHeap || bRequestNewActiveGpuHeap)
+	if (State.bRefreshHeap || bRequestNewActiveGpuHeap)
 	{
 		// First finalize the previous heap if it was set.
 		FinalizeHeapOnState(State);
@@ -514,11 +512,12 @@ bool FD3D12BindlessResourceManager::FlushPendingDescriptorUpdates(FD3D12CommandC
 		{
 			// Finally tell the Context that we're using this heap,
 			// this call also makes sure the heap is set on the d3d command list.
-			bSetNewHeaps = Context.StateCache.GetDescriptorCache()->SwitchToNewBindlessResourceHeap(State.CurrentGpuHeap);
+			if (Context.StateCache.GetDescriptorCache()->SwitchToNewBindlessResourceHeap(State.CurrentGpuHeap))
+			{
+				Context.StateCache.BindlessHeapsInvalidated();
+			}
 		}
 	}
-
-	return bSetNewHeaps;
 }
 
 void FD3D12BindlessResourceManager::OpenCommandList(FD3D12CommandContext& Context)
@@ -561,8 +560,6 @@ void FD3D12BindlessResourceManager::FinalizeContext(FD3D12CommandContext& Contex
 	// If context wasn't opened but did have descriptor updates make sure the shared gpu heap is updated
 	// (can happen due to texture reference updates not adding any real GPU work)
 	FinalizeHeapOnState(State);
-
-	check(!Context.GetBindlessState().HasAnyPending());
 }
 
 FD3D12DescriptorHeap* FD3D12BindlessResourceManager::GetHeap(ERHIPipeline Pipeline) const
@@ -627,7 +624,7 @@ void FD3D12BindlessResourceManager::FinalizeHeapOnState(FD3D12ContextBindlessSta
 
 	// Clear the state data
 	State.CurrentGpuHeap = nullptr;
-	State.bRequestNewGpuHeap = false;
+	State.bRefreshHeap = false;
 }
 
 void FD3D12BindlessResourceManager::CheckRequestNewActiveGPUHeap()
@@ -869,16 +866,12 @@ void FD3D12BindlessDescriptorManager::CloseCommandList(FD3D12CommandContext& Con
 	}
 }
 
-bool FD3D12BindlessDescriptorManager::FlushPendingDescriptorUpdates(FD3D12CommandContext& Context)
+void FD3D12BindlessDescriptorManager::FlushPendingDescriptorUpdates(FD3D12CommandContext& Context)
 {
-	bool bSetNewHeaps = false;
-
 	if (ResourceManager)
 	{
-		bSetNewHeaps = ResourceManager->FlushPendingDescriptorUpdates(Context);
+		ResourceManager->FlushPendingDescriptorUpdates(Context);
 	}
-
-	return bSetNewHeaps;
 }
 
 FD3D12DescriptorHeapPair FD3D12BindlessDescriptorManager::GetExplicitHeapsForContext(FD3D12CommandContext& Context, ERHIBindlessConfiguration InConfiguration)
diff --git a/Engine/Source/Runtime/D3D12RHI/Private/D3D12BindlessDescriptors.h b/Engine/Source/Runtime/D3D12RHI/Private/D3D12BindlessDescriptors.h
--- a/Engine/Source/Runtime/D3D12RHI/Private/D3D12BindlessDescriptors.h
+++ b/Engine/Source/Runtime/D3D12RHI/Private/D3D12BindlessDescriptors.h
@@ -103,17 +103,17 @@ private:
 struct FD3D12ContextBindlessState
 {
 	FD3D12DescriptorHeapPtr CurrentGpuHeap;
-	bool                    bRequestNewGpuHeap = false;
+	bool                    bRefreshHeap = false;
 
 	FD3D12ContextBindlessState() = default;
 	~FD3D12ContextBindlessState()
 	{
-		check(!bRequestNewGpuHeap);
+		check(!bRefreshHeap);
 	}
 
-	bool HasAnyPending() const
+	void RefreshDescriptorHeap()
 	{
-		return bRequestNewGpuHeap;
+		bRefreshHeap = true;
 	}
 };
 
@@ -176,7 +176,7 @@ public:
 	void InitializeDescriptor(FRHIDescriptorHandle DstHandle, FD3D12View* View);
 	void UpdateDescriptor(FD3D12ContextArray const& Contexts, FRHIDescriptorHandle DstHandle, FD3D12View* View);
 
-	bool FlushPendingDescriptorUpdates(FD3D12CommandContext& Context);
+	void FlushPendingDescriptorUpdates(FD3D12CommandContext& Context);
 
 	void OpenCommandList(FD3D12CommandContext& Context);
 	void CloseCommandList(FD3D12CommandContext& Context);
@@ -278,7 +278,7 @@ public:
 	void OpenCommandList(FD3D12CommandContext& Context);
 	void CloseCommandList(FD3D12CommandContext& Context);
 
-	bool FlushPendingDescriptorUpdates(FD3D12CommandContext& Context);
+	void FlushPendingDescriptorUpdates(FD3D12CommandContext& Context);
 	void SetHeapsForRayTracing(FD3D12CommandContext& Context);
 
 	FD3D12DescriptorHeapPair GetExplicitHeapsForContext(FD3D12CommandContext& Context, ERHIBindlessConfiguration InConfiguration);
diff --git a/Engine/Source/Runtime/D3D12RHI/Private/D3D12CommandContext.h b/Engine/Source/Runtime/D3D12RHI/Private/D3D12CommandContext.h
--- a/Engine/Source/Runtime/D3D12RHI/Private/D3D12CommandContext.h
+++ b/Engine/Source/Runtime/D3D12RHI/Private/D3D12CommandContext.h
@@ -797,7 +797,7 @@ public:
 		return StaticUniformBuffers;
 	}
 
-	bool FlushPendingDescriptorUpdates();
+	void FlushPendingDescriptorUpdates();
 
 	void SetExplicitDescriptorCache(FD3D12ExplicitDescriptorCache& ExplicitDescriptorCache);
 	void UnsetExplicitDescriptorCache();
diff --git a/Engine/Source/Runtime/D3D12RHI/Private/D3D12Commands.cpp b/Engine/Source/Runtime/D3D12RHI/Private/D3D12Commands.cpp
--- a/Engine/Source/Runtime/D3D12RHI/Private/D3D12Commands.cpp
+++ b/Engine/Source/Runtime/D3D12RHI/Private/D3D12Commands.cpp
@@ -73,17 +73,13 @@ static void BindUniformBuffer(FD3D12CommandContext& Context, FRHIShader* Shader,
 	Context.DirtyUniformBuffers[ShaderFrequency] |= (1 << BufferIndex);
 }
 
-bool FD3D12CommandContext::FlushPendingDescriptorUpdates()
+void FD3D12CommandContext::FlushPendingDescriptorUpdates()
 {
-	bool bNewBindlessHeapsSet = false;
-
 #if PLATFORM_SUPPORTS_BINDLESS_RENDERING
 	// Make sure the graphics command list is valid and open before trying to flush pending descriptor updates
-	GraphicsCommandList();
-	bNewBindlessHeapsSet = GetParentDevice()->GetBindlessDescriptorManager().FlushPendingDescriptorUpdates(*this);
+	OpenIfNotAlready();
+	GetParentDevice()->GetBindlessDescriptorManager().FlushPendingDescriptorUpdates(*this);
 #endif
-
-	return bNewBindlessHeapsSet;
 }
 
 void FD3D12CommandContext::SetExplicitDescriptorCache(FD3D12ExplicitDescriptorCache& ExplicitDescriptorCache)
@@ -113,12 +109,12 @@ void FD3D12CommandContext::SetupDispatch(uint32 ThreadGroupCountX, uint32 Thread
 	}
 #endif
 
-	const bool bNewBindlessHeapsSet = FlushPendingDescriptorUpdates();
+	FlushPendingDescriptorUpdates();
 
 	CommitComputeShaderConstants();
 	CommitComputeResourceTables();
 
-	StateCache.ApplyState(GetPipeline(), ED3D12PipelineType::Compute, bNewBindlessHeapsSet);
+	StateCache.ApplyState(GetPipeline(), ED3D12PipelineType::Compute);
 }
 
 FD3D12ResourceLocation& FD3D12CommandContext::SetupIndirectArgument(FRHIBuffer* ArgumentBufferRHI, D3D12_RESOURCE_STATES ExtraStates)
@@ -1758,7 +1754,7 @@ void FD3D12CommandContext::SetupDraw(FRHIBuffer* IndexBufferRHI, uint32 NumPrimi
 	}
 #endif
 
-	const bool bNewBindlessHeapsSet = FlushPendingDescriptorUpdates();
+	FlushPendingDescriptorUpdates();
 
 	CommitGraphicsResourceTables();
 	CommitNonComputeShaderConstants();
@@ -1773,7 +1769,7 @@ void FD3D12CommandContext::SetupDraw(FRHIBuffer* IndexBufferRHI, uint32 NumPrimi
 		StateCache.SetIndexBuffer(IndexBuffer->ResourceLocation, Format, 0);
 	}
 
-	StateCache.ApplyState(GetPipeline(), ED3D12PipelineType::Graphics, bNewBindlessHeapsSet);
+	StateCache.ApplyState(GetPipeline(), ED3D12PipelineType::Graphics);
 }
 
 void FD3D12CommandContext::SetupDispatchDraw(uint32 ThreadGroupCountX, uint32 ThreadGroupCountY, uint32 ThreadGroupCountZ)
@@ -1785,12 +1781,12 @@ void FD3D12CommandContext::SetupDispatchDraw(uint32 ThreadGroupCountX, uint32 Th
 	}
 #endif
 
-	const bool bNewBindlessHeapsSet = FlushPendingDescriptorUpdates();
+	FlushPendingDescriptorUpdates();
 
 	CommitGraphicsResourceTables();
 	CommitNonComputeShaderConstants();
 
-	StateCache.ApplyState(GetPipeline(), ED3D12PipelineType::Graphics, bNewBindlessHeapsSet);
+	StateCache.ApplyState(GetPipeline(), ED3D12PipelineType::Graphics);
 }
 
 void FD3D12CommandContext::RHIDrawPrimitive(uint32 BaseVertexIndex, uint32 NumPrimitives, uint32 NumInstances)
diff --git a/Engine/Source/Runtime/D3D12RHI/Private/D3D12StateCache.cpp b/Engine/Source/Runtime/D3D12RHI/Private/D3D12StateCache.cpp
--- a/Engine/Source/Runtime/D3D12RHI/Private/D3D12StateCache.cpp
+++ b/Engine/Source/Runtime/D3D12RHI/Private/D3D12StateCache.cpp
@@ -418,7 +418,7 @@ void FD3D12StateCache::InternalSetPipelineState(FD3D12PipelineState* InPipelineS
 	}
 }
 
-void FD3D12StateCache::ApplyState(ERHIPipeline HardwarePipe, ED3D12PipelineType PipelineType, bool bBindlessHeapsWereJustSet)
+void FD3D12StateCache::ApplyState(ERHIPipeline HardwarePipe, ED3D12PipelineType PipelineType)
 {
 	//SCOPE_CYCLE_COUNTER(STAT_D3D12ApplyStateTime);
 	const bool bForceState = false;
@@ -427,14 +427,6 @@ void FD3D12StateCache::ApplyState(ERHIPipeline HardwarePipe, ED3D12PipelineType
 		// Mark all state as dirty.
 		DirtyState();
 	}
-#if PLATFORM_SUPPORTS_BINDLESS_RENDERING
-	// If we just switched to a new bindless heap, we have to make sure to set the RootSignatures again.
-	else if (bBindlessHeapsWereJustSet)
-	{
-		PipelineState.Compute.bNeedSetRootSignature = true;
-		PipelineState.Graphics.bNeedSetRootSignature = true;
-	}
-#endif
 
 #if PLATFORM_SUPPORTS_VIRTUAL_TEXTURES
 	CmdContext.FlushTextureCacheIfNeeded();
diff --git a/Engine/Source/Runtime/D3D12RHI/Private/D3D12StateCachePrivate.h b/Engine/Source/Runtime/D3D12RHI/Private/D3D12StateCachePrivate.h
--- a/Engine/Source/Runtime/D3D12RHI/Private/D3D12StateCachePrivate.h
+++ b/Engine/Source/Runtime/D3D12RHI/Private/D3D12StateCachePrivate.h
@@ -675,7 +675,7 @@ public:
 	ED3D12PipelineType LastComputePipelineType = ED3D12PipelineType::Compute;
 #endif // D3D12_RHI_RAYTRACING
 
-	void ApplyState(ERHIPipeline HardwarePipe, ED3D12PipelineType PipelineType, bool bBindlessHeapsWereJustSet);
+	void ApplyState(ERHIPipeline HardwarePipe, ED3D12PipelineType PipelineType);
 	void ApplySamplers(const FD3D12RootSignature* const pRootSignature, uint32 StartStage, uint32 EndStage);
 	void ApplyResources(const FD3D12RootSignature* const pRootSignature, uint32 StartStage, uint32 EndStage);
 	void ApplyBindlessResources(const FD3D12RootSignature* const pRootSignature, uint32 StartStage, uint32 EndStage);
@@ -758,6 +758,13 @@ public:
 	void ForceSetGraphicsRootSignature() { PipelineState.Graphics.bNeedSetRootSignature = true; }
 
 #if PLATFORM_SUPPORTS_BINDLESS_RENDERING
+	void BindlessHeapsInvalidated()
+	{
+		// If we just switched to a new bindless heap, we have to make sure to set the RootSignatures again.
+		PipelineState.Compute.bNeedSetRootSignature = true;
+		PipelineState.Graphics.bNeedSetRootSignature = true;
+	}
+
 	void QueueBindlessSRV(EShaderFrequency ShaderFrequency, FD3D12ShaderResourceView* SRV)
 	{
 		PipelineState.Common.QueuedBindlessSRVs[ShaderFrequency].Emplace(SRV);
diff --git a/Engine/Source/Runtime/D3D12RHI/Private/D3D12UAV.cpp b/Engine/Source/Runtime/D3D12RHI/Private/D3D12UAV.cpp
--- a/Engine/Source/Runtime/D3D12RHI/Private/D3D12UAV.cpp
+++ b/Engine/Source/Runtime/D3D12RHI/Private/D3D12UAV.cpp
@@ -329,8 +329,10 @@ void FD3D12CommandContext::ClearUAV(TRHICommandList_RecursiveHazardous<FD3D12Com
 		if (UAVDesc.Format == DXGI_FORMAT_UNKNOWN || (UAVDesc.Buffer.Flags & D3D12_BUFFER_UAV_FLAG_RAW) != 0)
 		{
 			// Structured buffer.
-			RHICmdList.RunOnContext([UnorderedAccessView, ClearValues, UAVDesc](FD3D12CommandContext& Context)
+			RHICmdList.RunOnContext([UnorderedAccessView, ClearValues](FD3D12CommandContext& Context)
 			{
+				const D3D12_UNORDERED_ACCESS_VIEW_DESC& UAVDesc = UnorderedAccessView->GetD3DDesc();
+
 				// Alias the structured buffer with an R32_UINT UAV to perform the clear.
 				// We construct a temporary UAV on the offline heap, copy it to the online heap, and then call ClearUnorderedAccessViewUint.
 
@@ -367,6 +369,11 @@ void FD3D12CommandContext::ClearUAV(TRHICommandList_RecursiveHazardous<FD3D12Com
 				if (Context.StateCache.GetDescriptorCache()->IsActiveViewHeapBindless())
 				{
 					check(UAV.GetBindlessHandle().IsValid());
+
+					// We just allocated a descriptor inline. Since this could have recreated the descriptor heaps,
+					// force a refresh of the current GPU heap. If it didn't change, no extra commands should be added.
+					Context.GetBindlessState().RefreshDescriptorHeap();
+
 					Context.FlushPendingDescriptorUpdates();
 
 					FD3D12DescriptorHeap* BindlessHeap = Context.GetBindlessResourcesHeap();
