Removed final qualifier from FD3D12CommandContext::OpenCommandList

diff --git a/Engine/Source/Runtime/D3D12RHI/Private/D3D12CommandContext.h b/Engine/Source/Runtime/D3D12RHI/Private/D3D12CommandContext.h
--- a/Engine/Source/Runtime/D3D12RHI/Private/D3D12CommandContext.h
+++ b/Engine/Source/Runtime/D3D12RHI/Private/D3D12CommandContext.h
@@ -510,7 +510,7 @@ public:
 #endif
 	}
 
-	virtual void OpenCommandList() override final;
+	virtual void OpenCommandList() override;
 	virtual void CloseCommandList() override final;
 
 	virtual ERHIPipeline GetPipeline() const override
