[HWRT] Fix incorrect logic in D3D12 RHI when using r.RayTracing.OnDemandGeometryBuffersStreaming.
- only need vertex and index buffers to be valid and resident when using TLAS (tracing passes) if RHI supports hit shaders or requires InlineRayTracingSBT.

diff --git a/Engine/Source/Runtime/D3D12RHI/Private/D3D12RayTracing.cpp b/Engine/Source/Runtime/D3D12RHI/Private/D3D12RayTracing.cpp
--- a/Engine/Source/Runtime/D3D12RHI/Private/D3D12RayTracing.cpp
+++ b/Engine/Source/Runtime/D3D12RHI/Private/D3D12RayTracing.cpp
@@ -4492,23 +4492,27 @@ void BuildAccelerationStructure(FD3D12CommandContext& CommandContext,
 
 				checkf(!Geometry->IsDirty(CommandContext.GetGPUIndex()),
 					TEXT("Acceleration structures for all geometries must be built before building the top level acceleration structure for the scene."));
-				checkf(Geometry->BuffersValid(CommandContext.GetGPUIndex()),
-					TEXT("Index & vertex buffers for all geometries must be valid (streamed in) when adding geometry to the top level acceleration structure for the scene"));
 
 				AddResidencyHandleForResource(Geometry->AccelerationStructureBuffers[GPUIndex]->GetResource());
 
-				if (Geometry->Initializer.IndexBuffer)
+				if (GRHIGlobals.RayTracing.SupportsShaders || GRHIGlobals.RayTracing.RequiresInlineRayTracingSBT)
 				{
-					FD3D12Buffer* IndexBuffer = CommandContext.RetrieveObject<FD3D12Buffer>(Geometry->Initializer.IndexBuffer.GetReference());
-					AddResidencyHandleForResource(IndexBuffer->GetResource());
-				}
+					checkf(Geometry->BuffersValid(CommandContext.GetGPUIndex()),
+						TEXT("Index & vertex buffers for all geometries must be valid (streamed in) when adding geometry to the top level acceleration structure for the scene"));
 
-				for (const FRayTracingGeometrySegment& Segment : Geometry->Initializer.Segments)
-				{
-					if (Segment.VertexBuffer)
+					if (Geometry->Initializer.IndexBuffer)
 					{
-						FD3D12Buffer* VertexBuffer = CommandContext.RetrieveObject<FD3D12Buffer>(Segment.VertexBuffer.GetReference());
-						AddResidencyHandleForResource(VertexBuffer->GetResource());
+						FD3D12Buffer* IndexBuffer = CommandContext.RetrieveObject<FD3D12Buffer>(Geometry->Initializer.IndexBuffer.GetReference());
+						AddResidencyHandleForResource(IndexBuffer->GetResource());
+					}
+
+					for (const FRayTracingGeometrySegment& Segment : Geometry->Initializer.Segments)
+					{
+						if (Segment.VertexBuffer)
+						{
+							FD3D12Buffer* VertexBuffer = CommandContext.RetrieveObject<FD3D12Buffer>(Segment.VertexBuffer.GetReference());
+							AddResidencyHandleForResource(VertexBuffer->GetResource());
+						}
 					}
 				}
 			}
