Fix D3D12 CopyTextureRegionChecked() to take both source and dest pixel formats.
We need to know both to validate aliased copies between UINT and BC compressed formats.
#rnx

diff --git a/Engine/Source/Runtime/D3D12RHI/Private/D3D12CommandContext.cpp b/Engine/Source/Runtime/D3D12RHI/Private/D3D12CommandContext.cpp
--- a/Engine/Source/Runtime/D3D12RHI/Private/D3D12CommandContext.cpp
+++ b/Engine/Source/Runtime/D3D12RHI/Private/D3D12CommandContext.cpp
@@ -1055,7 +1055,7 @@ void FD3D12CommandContextRedirector::RHITransferResources(TConstArrayView<FTrans
 				CD3DX12_TEXTURE_COPY_LOCATION SrcLocation(SrcTexture->GetResource()->GetResource(), 0);
 				CD3DX12_TEXTURE_COPY_LOCATION DstLocation(DstTexture->GetResource()->GetResource(), 0);
 
-				CopyContext->CopyTextureRegionChecked(&DstLocation, Box.left, Box.top, Box.front, &SrcLocation, &Box, DstTexture->GetFormat(), DstTexture->GetName());
+				CopyContext->CopyTextureRegionChecked(&DstLocation, Box.left, Box.top, Box.front, DstTexture->GetFormat(), &SrcLocation, &Box, SrcTexture->GetFormat(), DstTexture->GetName());
 			}
 		}
 		else
@@ -1209,7 +1209,7 @@ void FD3D12CommandContextRedirector::RHICrossGPUTransfer(TConstArrayView<FTransf
 				CD3DX12_TEXTURE_COPY_LOCATION SrcLocation(SrcTexture->GetResource()->GetResource(), 0);
 				CD3DX12_TEXTURE_COPY_LOCATION DstLocation(DstTexture->GetResource()->GetResource(), 0);
 
-				SrcContext->CopyTextureRegionChecked(&DstLocation, Box.left, Box.top, Box.front, &SrcLocation, &Box, DstTexture->GetFormat(), DstTexture->GetName());
+				SrcContext->CopyTextureRegionChecked(&DstLocation, Box.left, Box.top, Box.front, DstTexture->GetFormat(), &SrcLocation, &Box, SrcTexture->GetFormat(), DstTexture->GetName());
 			}
 		}
 		else
diff --git a/Engine/Source/Runtime/D3D12RHI/Private/D3D12CommandContext.h b/Engine/Source/Runtime/D3D12RHI/Private/D3D12CommandContext.h
--- a/Engine/Source/Runtime/D3D12RHI/Private/D3D12CommandContext.h
+++ b/Engine/Source/Runtime/D3D12RHI/Private/D3D12CommandContext.h
@@ -832,9 +832,8 @@ public:
 	);
 	
 	FORCENOINLINE void CopyTextureRegionChecked(
-		const D3D12_TEXTURE_COPY_LOCATION* DestCopyLocation, int DestX, int DestY, int DestZ,
-		const D3D12_TEXTURE_COPY_LOCATION* SourceCopyLocation, const D3D12_BOX* SourceBox,
-		EPixelFormat PixelFormat,
+		const D3D12_TEXTURE_COPY_LOCATION* DestCopyLocation, int DestX, int DestY, int DestZ, EPixelFormat DestPixelFormat,
+		const D3D12_TEXTURE_COPY_LOCATION* SourceCopyLocation, const D3D12_BOX* SourceBox, EPixelFormat SourcePixelFormat,
 		const FName& DebugName
 	);
 
diff --git a/Engine/Source/Runtime/D3D12RHI/Private/D3D12RenderTarget.cpp b/Engine/Source/Runtime/D3D12RHI/Private/D3D12RenderTarget.cpp
--- a/Engine/Source/Runtime/D3D12RHI/Private/D3D12RenderTarget.cpp
+++ b/Engine/Source/Runtime/D3D12RHI/Private/D3D12RenderTarget.cpp
@@ -482,6 +482,7 @@ TRefCountPtr<FD3D12Resource> FD3D12DynamicRHI::GetStagingTexture(FRHITexture* Te
 			Context.CopyTextureRegionChecked(
 				&DestCopyLocation,
 				0, 0, 0,
+				Texture->GetFormat(),
 				&SourceCopyLocation,
 				RectPtr,
 				Texture->GetFormat(),
@@ -801,6 +802,7 @@ void FD3D12DynamicRHI::ReadSurfaceDataMSAARaw(FRHITexture* TextureRHI, FIntRect
 			Context.CopyTextureRegionChecked(
 				&DestCopyLocation,
 				0, 0, 0,
+				Texture->GetFormat(),
 				&SourceCopyLocation,
 				&Rect,
 				Texture->GetFormat(),
@@ -979,6 +981,7 @@ void FD3D12DynamicRHI::RHIReadSurfaceFloatData(FRHITexture* TextureRHI, FIntRect
 		Context.CopyTextureRegionChecked(
 			&DestCopyLocation,
 			0, 0, 0,
+			Texture->GetFormat(),
 			&SourceCopyLocation,
 			&Rect,
 			Texture->GetFormat(),
@@ -1093,6 +1096,7 @@ void FD3D12DynamicRHI::RHIRead3DSurfaceFloatData(FRHITexture* TextureRHI, FIntRe
 		Context.CopyTextureRegionChecked(
 			&DestCopyLocation,
 			0, 0, 0,
+			Texture->GetFormat(),
 			&SourceCopyLocation,
 			&Rect,
 			Texture->GetFormat(),
diff --git a/Engine/Source/Runtime/D3D12RHI/Private/D3D12Texture.cpp b/Engine/Source/Runtime/D3D12RHI/Private/D3D12Texture.cpp
--- a/Engine/Source/Runtime/D3D12RHI/Private/D3D12Texture.cpp
+++ b/Engine/Source/Runtime/D3D12RHI/Private/D3D12Texture.cpp
@@ -1192,9 +1192,8 @@ static uint64 GetPlacedBufferRequiredSize(FD3D12CommandContext& Context, const D
 }
 
 FORCENOINLINE void FD3D12CommandContext::CopyTextureRegionChecked(
-	const D3D12_TEXTURE_COPY_LOCATION* DestCopyLocation, int DestX, int DestY, int DestZ,
-	const D3D12_TEXTURE_COPY_LOCATION* SourceCopyLocation, const D3D12_BOX* SourceBox,
-	EPixelFormat PixelFormat,
+	const D3D12_TEXTURE_COPY_LOCATION* DestCopyLocation, int DestX, int DestY, int DestZ, EPixelFormat DestPixelFormat,
+	const D3D12_TEXTURE_COPY_LOCATION* SourceCopyLocation, const D3D12_BOX* SourceBox, EPixelFormat SourcePixelFormat,
 	const FName& DebugName
 )
 {
@@ -1202,7 +1201,8 @@ FORCENOINLINE void FD3D12CommandContext::CopyTextureRegionChecked(
 	const D3D12_RESOURCE_DESC DestDesc   = DestCopyLocation->pResource->GetDesc();
 	const D3D12_RESOURCE_DESC SourceDesc = SourceCopyLocation->pResource->GetDesc();
 
-	const FPixelFormatInfo& PixelFormatInfo = GPixelFormats[PixelFormat];
+	const FPixelFormatInfo& DestPixelFormatInfo = GPixelFormats[DestPixelFormat];
+	const FPixelFormatInfo& SourcePixelFormatInfo = GPixelFormats[SourcePixelFormat];
 
 	switch (SourceCopyLocation->Type)
 	{
@@ -1213,12 +1213,12 @@ FORCENOINLINE void FD3D12CommandContext::CopyTextureRegionChecked(
 
 			const D3D12_BOX CopyBox = SourceBox ? *SourceBox : GetSubresourceCopyBox(SourceDesc, DestCopyLocation, Indices);
 
-			RHI_BREADCRUMB_CHECK_SHIPPINGF(*this, (CopyBox.left % PixelFormatInfo.BlockSizeX) == 0 && (CopyBox.right % PixelFormatInfo.BlockSizeX) == 0, TEXT("Width not aligned to block size for: '%s'"), *DebugName.ToString());
-			RHI_BREADCRUMB_CHECK_SHIPPINGF(*this, (CopyBox.top % PixelFormatInfo.BlockSizeY) == 0 && (CopyBox.bottom % PixelFormatInfo.BlockSizeY) == 0, TEXT("Height not aligned to block size for: '%s'"), *DebugName.ToString());
+			RHI_BREADCRUMB_CHECK_SHIPPINGF(*this, (CopyBox.left % SourcePixelFormatInfo.BlockSizeX) == 0 && (CopyBox.right % SourcePixelFormatInfo.BlockSizeX) == 0, TEXT("Width not aligned to block size for: '%s'"), *DebugName.ToString());
+			RHI_BREADCRUMB_CHECK_SHIPPINGF(*this, (CopyBox.top % SourcePixelFormatInfo.BlockSizeY) == 0 && (CopyBox.bottom % SourcePixelFormatInfo.BlockSizeY) == 0, TEXT("Height not aligned to block size for: '%s'"), *DebugName.ToString());
 				
-			const uint32 SubresourceWidth  = AlignArbitrary<uint32>(FMath::Max<uint32>(1u, SourceDesc.Width >> Indices.MipIndex), PixelFormatInfo.BlockSizeX);
-			const uint32 SubresourceHeight = AlignArbitrary<uint32>(FMath::Max<uint32>(1u, SourceDesc.Height >> Indices.MipIndex), PixelFormatInfo.BlockSizeY);
-			const uint32 SubresourceDepth  = AlignArbitrary<uint32>(SourceDesc.Dimension == D3D12_RESOURCE_DIMENSION_TEXTURE3D ? FMath::Max<uint32>(1u, SourceDesc.DepthOrArraySize >> Indices.MipIndex) : 1, PixelFormatInfo.BlockSizeZ);
+			const uint32 SubresourceWidth  = AlignArbitrary<uint32>(FMath::Max<uint32>(1u, SourceDesc.Width >> Indices.MipIndex), SourcePixelFormatInfo.BlockSizeX);
+			const uint32 SubresourceHeight = AlignArbitrary<uint32>(FMath::Max<uint32>(1u, SourceDesc.Height >> Indices.MipIndex), SourcePixelFormatInfo.BlockSizeY);
+			const uint32 SubresourceDepth  = AlignArbitrary<uint32>(SourceDesc.Dimension == D3D12_RESOURCE_DIMENSION_TEXTURE3D ? FMath::Max<uint32>(1u, SourceDesc.DepthOrArraySize >> Indices.MipIndex) : 1, SourcePixelFormatInfo.BlockSizeZ);
 
 			RHI_BREADCRUMB_CHECK_SHIPPINGF(*this, CopyBox.left <= CopyBox.right && CopyBox.right <= SubresourceWidth, TEXT("Source width out of bounds for: '%s'"), *DebugName.ToString());
 			RHI_BREADCRUMB_CHECK_SHIPPINGF(*this, CopyBox.top <= CopyBox.bottom && CopyBox.bottom <= SubresourceHeight, TEXT("Source height out of bounds for: '%s'"), *DebugName.ToString());
@@ -1229,7 +1229,7 @@ FORCENOINLINE void FD3D12CommandContext::CopyTextureRegionChecked(
 		// Copy from buffer
 		case D3D12_TEXTURE_COPY_TYPE_PLACED_FOOTPRINT:
 		{
-			const uint32 RequiredPlacedSize = GetPlacedBufferRequiredSize(*this, SourceCopyLocation->PlacedFootprint.Footprint, PixelFormatInfo, DebugName);
+			const uint32 RequiredPlacedSize = GetPlacedBufferRequiredSize(*this, SourceCopyLocation->PlacedFootprint.Footprint, SourcePixelFormatInfo, DebugName);
 			RHI_BREADCRUMB_CHECK_SHIPPINGF(*this, SourceCopyLocation->PlacedFootprint.Offset + RequiredPlacedSize <= SourceDesc.Width, TEXT("Source placed buffer width out of bounds for: '%s'"), *DebugName.ToString());
 			break;
 		}
@@ -1244,12 +1244,12 @@ FORCENOINLINE void FD3D12CommandContext::CopyTextureRegionChecked(
 
 			const D3D12_BOX CopyBox = SourceBox ? *SourceBox : GetSubresourceCopyBox(DestDesc, SourceCopyLocation, Indices);
 				
-			RHI_BREADCRUMB_CHECK_SHIPPINGF(*this, (CopyBox.left % PixelFormatInfo.BlockSizeX) == 0 && (CopyBox.right % PixelFormatInfo.BlockSizeX) == 0, TEXT("Width not aligned to block size for: '%s'"), *DebugName.ToString());
-			RHI_BREADCRUMB_CHECK_SHIPPINGF(*this, (CopyBox.top % PixelFormatInfo.BlockSizeY) == 0 && (CopyBox.bottom % PixelFormatInfo.BlockSizeY) == 0, TEXT("Height not aligned to block size for: '%s'"), *DebugName.ToString());
+			RHI_BREADCRUMB_CHECK_SHIPPINGF(*this, (CopyBox.left % DestPixelFormatInfo.BlockSizeX) == 0 && (CopyBox.right % DestPixelFormatInfo.BlockSizeX) == 0, TEXT("Width not aligned to block size for: '%s'"), *DebugName.ToString());
+			RHI_BREADCRUMB_CHECK_SHIPPINGF(*this, (CopyBox.top % DestPixelFormatInfo.BlockSizeY) == 0 && (CopyBox.bottom % DestPixelFormatInfo.BlockSizeY) == 0, TEXT("Height not aligned to block size for: '%s'"), *DebugName.ToString());
 				
-			const uint32 SubresourceWidth  = AlignArbitrary<uint32>(FMath::Max<uint32>(1u, DestDesc.Width >> Indices.MipIndex), PixelFormatInfo.BlockSizeX);
-			const uint32 SubresourceHeight = AlignArbitrary<uint32>(FMath::Max<uint32>(1u, DestDesc.Height >> Indices.MipIndex), PixelFormatInfo.BlockSizeY);
-			const uint32 SubresourceDepth  = AlignArbitrary<uint32>(DestDesc.Dimension == D3D12_RESOURCE_DIMENSION_TEXTURE3D ? FMath::Max<uint32>(1u, DestDesc.DepthOrArraySize >> Indices.MipIndex) : 1, PixelFormatInfo.BlockSizeZ);
+			const uint32 SubresourceWidth  = AlignArbitrary<uint32>(FMath::Max<uint32>(1u, DestDesc.Width >> Indices.MipIndex), DestPixelFormatInfo.BlockSizeX);
+			const uint32 SubresourceHeight = AlignArbitrary<uint32>(FMath::Max<uint32>(1u, DestDesc.Height >> Indices.MipIndex), DestPixelFormatInfo.BlockSizeY);
+			const uint32 SubresourceDepth  = AlignArbitrary<uint32>(DestDesc.Dimension == D3D12_RESOURCE_DIMENSION_TEXTURE3D ? FMath::Max<uint32>(1u, DestDesc.DepthOrArraySize >> Indices.MipIndex) : 1, DestPixelFormatInfo.BlockSizeZ);
 
 			RHI_BREADCRUMB_CHECK_SHIPPINGF(*this, DestX + (CopyBox.right - CopyBox.left) <= SubresourceWidth, TEXT("Dest width out of bounds for: '%s'"), *DebugName.ToString());
 			RHI_BREADCRUMB_CHECK_SHIPPINGF(*this, DestY + (CopyBox.bottom - CopyBox.top) <= SubresourceHeight, TEXT("Dest height out of bounds for: '%s'"), *DebugName.ToString());
@@ -1260,7 +1260,7 @@ FORCENOINLINE void FD3D12CommandContext::CopyTextureRegionChecked(
 		// Copy to buffer
 		case D3D12_TEXTURE_COPY_TYPE_PLACED_FOOTPRINT:
 		{
-			const uint32 RequiredPlacedSize = GetPlacedBufferRequiredSize(*this, DestCopyLocation->PlacedFootprint.Footprint, PixelFormatInfo, DebugName);
+			const uint32 RequiredPlacedSize = GetPlacedBufferRequiredSize(*this, DestCopyLocation->PlacedFootprint.Footprint, DestPixelFormatInfo, DebugName);
 			RHI_BREADCRUMB_CHECK_SHIPPINGF(*this, DestCopyLocation->PlacedFootprint.Offset + RequiredPlacedSize <= DestDesc.Width, TEXT("Dest placed buffer width out of bounds for: '%s'"), *DebugName.ToString());
 			break;
 		}
@@ -1331,9 +1331,10 @@ FTextureRHIRef FD3D12DynamicRHI::AsyncReallocateTexture2D_RenderThread(FRHIComma
 				Context.CopyTextureRegionChecked(
 					&DestCopyLocation,
 					0, 0, 0,
+					DeviceNewTexture->GetFormat(),
 					&SourceCopyLocation,
 					nullptr,
-					DeviceNewTexture->GetFormat(),
+					DeviceOldTexture->GetFormat(),
 					DeviceNewTexture->GetName()
 				);
 
@@ -1829,6 +1830,7 @@ FRHILockTextureResult FD3D12Texture::Lock(FRHICommandListBase& RHICmdList, const
 			Context.CopyTextureRegionChecked(
 				&DestCopyLocation,
 				0, 0, 0,
+				GetFormat(),
 				&SourceCopyLocation,
 				nullptr,
 				GetFormat(),
@@ -1868,6 +1870,7 @@ void FD3D12Texture::UpdateTexture(FD3D12CommandContext& Context, uint32 MipIndex
 	Context.CopyTextureRegionChecked(
 		&DestCopyLocation,
 		DestX, DestY, DestZ,
+		GetFormat(),
 		&SourceCopyLocation,
 		nullptr,
 		GetFormat(),
@@ -1893,9 +1896,10 @@ void FD3D12Texture::CopyTextureRegion(FD3D12CommandContext& Context, uint32 Dest
 	Context.CopyTextureRegionChecked(
 		&DestCopyLocation,
 		DestX, DestY, DestZ,
+		GetFormat(),
 		&SourceCopyLocation,
 		&SourceBox,
-		GetFormat(),
+		SourceTexture->GetFormat(),
 		GetName()
 	);
 
@@ -2012,7 +2016,7 @@ void FD3D12Texture::InitializeTextureData(FRHICommandListBase& RHICmdList, const
 			{
 				Dst.SubresourceIndex = Subresource;
 				Src.PlacedFootprint = Footprints[Subresource];
-				Context.CopyTextureRegionChecked(&Dst, 0, 0, 0, &Src, nullptr, Texture->GetFormat(), Texture->GetName());
+				Context.CopyTextureRegionChecked(&Dst, 0, 0, 0, Texture->GetFormat(), &Src, nullptr, Texture->GetFormat(), Texture->GetName());
 			}
 
 			// Update the resource state after the copy has been done (will take care of updating the residency as well)
@@ -2348,6 +2352,7 @@ public:
 					UpdateInfo.DstStartX,
 					UpdateInfo.DstStartY,
 					UpdateInfo.DstStartZ,
+					TextureLink.GetFormat(),
 					&SourceCopyLocation,
 					nullptr,
 					TextureLink.GetFormat(),
@@ -2582,6 +2587,7 @@ void FD3D12DynamicRHI::EndUpdateTexture3D_Internal(FRHICommandListBase& RHICmdLi
 					DstStartX,
 					DstStartY,
 					DstStartZ,
+					Texture.GetFormat(),
 					&SourceCopyLocation,
 					nullptr,
 					Texture.GetFormat(),
@@ -2911,9 +2917,10 @@ void FD3D12CommandContext::RHICopyTexture(FRHITexture* SourceTextureRHI, FRHITex
 				CopyTextureRegionChecked(
 					&Dst,
 					DestX, DestY, DestZ,
+					DestTexture->GetFormat(),
 					&Src,
 					&SrcBox,
-					bReadback ? SourceTexture->GetFormat() : DestTexture->GetFormat(),
+					SourceTexture->GetFormat(),
 					SourceTexture->GetName()
 				);
 			}
