Additional Insights markers for D3D12 end-of-pipe tasks

diff --git a/Engine/Source/Runtime/D3D12RHI/Private/D3D12RHI.cpp b/Engine/Source/Runtime/D3D12RHI/Private/D3D12RHI.cpp
--- a/Engine/Source/Runtime/D3D12RHI/Private/D3D12RHI.cpp
+++ b/Engine/Source/Runtime/D3D12RHI/Private/D3D12RHI.cpp
@@ -395,6 +395,8 @@ void FD3D12DynamicRHI::RHIFlushResources()
 
 void FD3D12DynamicRHI::EnqueueEndOfPipeTask(TUniqueFunction<void()> TaskFunc, TUniqueFunction<void(FD3D12Payload&)> ModifyPayloadCallback)
 {
+	SCOPED_NAMED_EVENT_TEXT("EnqueueEndOfPipeTask", FColor::Yellow);
+
 	FGraphEventArray Prereqs;
 	Prereqs.Reserve(GD3D12MaxNumQueues + 1);
 	if (EopTask)
@@ -422,7 +424,11 @@ void FD3D12DynamicRHI::EnqueueEndOfPipeTask(TUniqueFunction<void()> TaskFunc, TU
 	SubmitPayloads(MoveTemp(Payloads));
 
 	EopTask = FFunctionGraphTask::CreateAndDispatchWhenReady(
-		MoveTemp(TaskFunc),
+		[Func = MoveTemp(TaskFunc)]()
+		{
+			SCOPED_NAMED_EVENT_TEXT("EndOfPipeTask", FColor::Red);
+			Func();
+		},
 		QUICK_USE_CYCLE_STAT(FExecuteRHIThreadTask, STATGROUP_TaskGraphTasks),
 		&Prereqs
 	);
@@ -442,6 +448,8 @@ void FD3D12DynamicRHI::RHIProcessDeleteQueue()
 	{
 		EnqueueEndOfPipeTask([Array = MoveTemp(Local)]()
 		{
+			SCOPED_NAMED_EVENT_TEXT("EndOfPipeTask_RHIProcessDeleteQueue", FColor::Silver);
+
 			for (FD3D12DeferredDeleteObject const& ObjectToDelete : Array)
 			{
 				switch (ObjectToDelete.Type)
@@ -592,6 +600,7 @@ void FD3D12DynamicRHI::RHIEndFrame(const FRHIEndFrameArgs& Args)
 	// Close the previous frame's timing and start a new one
 	auto Lambda = [this, OldTiming = MoveTemp(CurrentTimingPerQueue)]()
 	{
+		SCOPED_NAMED_EVENT_TEXT("EndOfPipeTask_RHIEndFrame", FColor::Orange);
 		ProcessTimestamps(OldTiming);
 	};
 
