Fix and Re-Enable Compute Based 3D Texture Updates For D3D12 Platforms

#tests EngineTest - "Rendering.RHI.Test RHI Copy.RHICopyTexture"

diff --git a/Engine/Source/Runtime/D3D12RHI/Private/D3D12Commands.cpp b/Engine/Source/Runtime/D3D12RHI/Private/D3D12Commands.cpp
--- a/Engine/Source/Runtime/D3D12RHI/Private/D3D12Commands.cpp
+++ b/Engine/Source/Runtime/D3D12RHI/Private/D3D12Commands.cpp
@@ -1031,13 +1031,13 @@ void FD3D12CommandContext::RHISetComputePipelineState(FRHIComputePipelineState*
 void FD3D12CommandContext::SetUAVParameter(EShaderFrequency Frequency, uint32 UAVIndex, FD3D12UnorderedAccessView* UAV)
 {
 	ClearShaderResources(UAV, EShaderParameterTypeMask::SRVMask);
-	StateCache.SetUAV(SF_Pixel, UAVIndex, UAV);
+	StateCache.SetUAV(Frequency, UAVIndex, UAV);
 }
 
 void FD3D12CommandContext::SetUAVParameter(EShaderFrequency Frequency, uint32 UAVIndex, FD3D12UnorderedAccessView* UAV, uint32 InitialCount)
 {
 	ClearShaderResources(UAV, EShaderParameterTypeMask::SRVMask);
-	StateCache.SetUAV(SF_Pixel, UAVIndex, UAV, InitialCount);
+	StateCache.SetUAV(Frequency, UAVIndex, UAV, InitialCount);
 }
 
 void FD3D12CommandContext::SetSRVParameter(EShaderFrequency Frequency, uint32 SRVIndex, FD3D12ShaderResourceView* SRV)
