UAV Access Workaround Refactor

Since the UAV access workaround is no longer needed for PC we can take advantage of other platforms and
avoid creating entire separate resources and instead directly create UAVs that override the format however
we'd like. This eliminates some complication in the barrier logic as well as should reduce the memory burden
somewhat.

Additionally, since these resources will have the UAV flag set directly on them now, we loosen the
retrictions on determining 4K alignment to not exclude UAVs. AFAICT this was never a D3D12
requirement and the validation does not complain on any supported D3D12 platform.

diff --git a/Engine/Source/Runtime/D3D12RHI/Private/D3D12Adapter.h b/Engine/Source/Runtime/D3D12RHI/Private/D3D12Adapter.h
--- a/Engine/Source/Runtime/D3D12RHI/Private/D3D12Adapter.h
+++ b/Engine/Source/Runtime/D3D12RHI/Private/D3D12Adapter.h
@@ -387,10 +387,9 @@ public:
 		ID3D12ResourceAllocator* ResourceAllocator = nullptr
 	);
 
-	void CreateUAVAliasResource(
-		D3D12_CLEAR_VALUE* ClearValuePtr,
-		const TCHAR* DebugName,
-		FD3D12ResourceLocation& Location);
+#if !D3D12RHI_SUPPORTS_UNCOMPRESSED_UAV
+	void CreateUAVAliasResourceDesc(const FD3D12ResourceLocation& Location);
+#endif
 
 	template <typename ObjectType, typename CreationCoreFunction>
 	inline ObjectType* CreateLinkedObject(FRHIGPUMask GPUMask, const CreationCoreFunction& pfnCreationCore)
diff --git a/Engine/Source/Runtime/D3D12RHI/Private/D3D12Allocation.cpp b/Engine/Source/Runtime/D3D12RHI/Private/D3D12Allocation.cpp
--- a/Engine/Source/Runtime/D3D12RHI/Private/D3D12Allocation.cpp
+++ b/Engine/Source/Runtime/D3D12RHI/Private/D3D12Allocation.cpp
@@ -2001,22 +2001,15 @@ FD3D12TextureAllocatorPool::FD3D12TextureAllocatorPool(FD3D12Device* Device, FRH
 	FD3D12DeviceChild(Device),
 	FD3D12MultiNodeGPUObject(Device->GetGPUMask(), VisibilityNode),
 	TraceHeapId(MemoryTrace_HeapSpec(EMemoryTraceRootHeap::VideoMemory, TEXT("Texture Allocator Pool"))),
-	ReadOnlyTexturePool(Device, VisibilityNode, FString(L"Small Read-Only Texture allocator"), TEXTURE_POOL_SIZE, D3D12_HEAP_FLAG_ALLOW_ONLY_NON_RT_DS_TEXTURES, TraceHeapId)
+	ReadOnlyTexturePool(Device, VisibilityNode, FString(L"Small Non-DS/RT Texture allocator"), TEXTURE_POOL_SIZE, D3D12_HEAP_FLAG_ALLOW_ONLY_NON_RT_DS_TEXTURES, TraceHeapId)
 {};
 
 HRESULT FD3D12TextureAllocatorPool::AllocateTexture(FD3D12ResourceDesc Desc, const D3D12_CLEAR_VALUE* ClearValue, EPixelFormat UEFormat, FD3D12ResourceLocation& TextureLocation, const D3D12_RESOURCE_STATES InitialState, const TCHAR* Name)
 {
-	// 4KB alignment is only available for read only textures
-	if (!EnumHasAnyFlags(Desc.Flags, D3D12_RESOURCE_FLAG_ALLOW_RENDER_TARGET|D3D12_RESOURCE_FLAG_ALLOW_DEPTH_STENCIL|D3D12_RESOURCE_FLAG_ALLOW_UNORDERED_ACCESS) &&
-		!Desc.NeedsUAVAliasWorkarounds() &&
-		Desc.SampleDesc.Count == 1)// Multi-Sample texures have much larger alignment requirements (4MB vs 64KB)
+	if (FD3D12Texture::CanBe4KAligned(Desc, (EPixelFormat)UEFormat))
 	{
-		// The top mip level must be less than 64k
-		if (FD3D12Texture::CanBe4KAligned(Desc, (EPixelFormat)UEFormat))
-		{
-			Desc.Alignment = D3D12_SMALL_RESOURCE_PLACEMENT_ALIGNMENT; // request 4k alignment
-			return ReadOnlyTexturePool.AllocateTexture(Desc, ClearValue, TextureLocation, InitialState, Name);
-		}
+		Desc.Alignment = D3D12_SMALL_RESOURCE_PLACEMENT_ALIGNMENT; // request 4k alignment
+		return ReadOnlyTexturePool.AllocateTexture(Desc, ClearValue, TextureLocation, InitialState, Name);
 	}
 
 	FD3D12Adapter* Adapter = GetParentDevice()->GetParentAdapter();
@@ -2025,35 +2018,6 @@ HRESULT FD3D12TextureAllocatorPool::AllocateTexture(FD3D12ResourceDesc Desc, con
 	const D3D12_HEAP_PROPERTIES HeapProps = CD3DX12_HEAP_PROPERTIES(D3D12_HEAP_TYPE_DEFAULT, GetGPUMask().GetNative(), GetVisibilityMask().GetNative());
 	const D3D12_RESOURCE_ALLOCATION_INFO Info = GetParentDevice()->GetResourceAllocationInfoUncached(Desc);
 
-	// UAV Aliasing needs a Heap to create the aliased resource in.
-	if (Desc.NeedsUAVAliasWorkarounds())
-	{
-		D3D12_HEAP_DESC HeapDesc{};
-		HeapDesc.SizeInBytes = Info.SizeInBytes;
-		HeapDesc.Properties = HeapProps;
-		HeapDesc.Alignment = Desc.SampleDesc.Count > 1 ? D3D12_DEFAULT_MSAA_RESOURCE_PLACEMENT_ALIGNMENT : 0;
-		HeapDesc.Flags = D3D12_HEAP_FLAG_ALLOW_ALL_BUFFERS_AND_TEXTURES;
-		if (Adapter->IsHeapNotZeroedSupported())
-		{
-			HeapDesc.Flags |= FD3D12_HEAP_FLAG_CREATE_NOT_ZEROED;
-		}
-
-		ID3D12Heap* Heap = nullptr;
-		VERIFYD3D12RESULT(Adapter->GetD3DDevice()->CreateHeap(&HeapDesc, IID_PPV_ARGS(&Heap)));
-		TRefCountPtr<FD3D12Heap> BackingHeap = new FD3D12Heap(GetParentDevice(), GetVisibilityMask(), TraceHeapId);
-		bool bTrack = false;
-		BackingHeap->SetHeap(Heap, Name, bTrack);
-
-		HRESULT hr = Adapter->CreatePlacedResource(Desc, BackingHeap, 0, InitialState, ED3D12ResourceStateMode::MultiState, InitialState, ClearValue, &Resource, Name);
-
-		if (SUCCEEDED(hr))
-		{
-			TextureLocation.AsStandAlone(Resource, Info.SizeInBytes);
-		}
-
-		return hr;
-	}
-
 	HRESULT hr = Adapter->CreateCommittedResource(Desc, GetGPUMask(), HeapProps, InitialState, ClearValue, &Resource, Name, false);
 
 	if (SUCCEEDED(hr))
diff --git a/Engine/Source/Runtime/D3D12RHI/Private/D3D12CommandContext.h b/Engine/Source/Runtime/D3D12RHI/Private/D3D12CommandContext.h
--- a/Engine/Source/Runtime/D3D12RHI/Private/D3D12CommandContext.h
+++ b/Engine/Source/Runtime/D3D12RHI/Private/D3D12CommandContext.h
@@ -426,10 +426,10 @@ public:
 	void FlushResourceBarriers();
 
 	// Functions for transitioning a resource. The Before and After state cannot be D3D12_RESOURCE_STATE_TBD, needs be known at the time of the call
-	bool TransitionResource(FD3D12Resource* Resource, D3D12_RESOURCE_STATES Before, D3D12_RESOURCE_STATES After, uint32 Subresource);
+	void TransitionResource(FD3D12Resource* Resource, D3D12_RESOURCE_STATES Before, D3D12_RESOURCE_STATES After, uint32 Subresource);
 	
 private:
-	bool TransitionResource(FD3D12Resource* InResource, uint32 InSubresourceIndex, D3D12_RESOURCE_STATES InBeforeState, D3D12_RESOURCE_STATES InAfterState);
+	void TransitionResource(FD3D12Resource* InResource, uint32 InSubresourceIndex, D3D12_RESOURCE_STATES InBeforeState, D3D12_RESOURCE_STATES InAfterState);
 };
 
 //
diff --git a/Engine/Source/Runtime/D3D12RHI/Private/D3D12CommandList.cpp b/Engine/Source/Runtime/D3D12RHI/Private/D3D12CommandList.cpp
--- a/Engine/Source/Runtime/D3D12RHI/Private/D3D12CommandList.cpp
+++ b/Engine/Source/Runtime/D3D12RHI/Private/D3D12CommandList.cpp
@@ -450,7 +450,7 @@ FD3D12CommandList::FState::FState(FD3D12CommandAllocator* CommandAllocator, FD3D
 	}
 }
 
-bool FD3D12ContextCommon::TransitionResource(FD3D12Resource* Resource, D3D12_RESOURCE_STATES Before, D3D12_RESOURCE_STATES After, uint32 Subresource)
+void FD3D12ContextCommon::TransitionResource(FD3D12Resource* Resource, D3D12_RESOURCE_STATES Before, D3D12_RESOURCE_STATES After, uint32 Subresource)
 {
 	check(Resource);
 	check(Resource->RequiresResourceStateTracking());
@@ -470,14 +470,10 @@ bool FD3D12ContextCommon::TransitionResource(FD3D12Resource* Resource, D3D12_RES
 #endif
 	UpdateResidency(Resource);
 
-	bool bRequireUAVBarrier = false;
-
-	bRequireUAVBarrier = TransitionResource(Resource, Subresource, Before, After);
-
-	return bRequireUAVBarrier;
+	TransitionResource(Resource, Subresource, Before, After);
 }
 
-static inline bool IsTransitionNeeded(D3D12_RESOURCE_STATES Before, D3D12_RESOURCE_STATES& After, FD3D12Resource* InResource = nullptr)
+static inline bool IsTransitionNeeded(D3D12_RESOURCE_STATES Before, D3D12_RESOURCE_STATES After, FD3D12Resource* InResource = nullptr)
 {
 	check(Before != D3D12_RESOURCE_STATE_CORRUPT && After != D3D12_RESOURCE_STATE_CORRUPT);
 	check(Before != D3D12_RESOURCE_STATE_TBD && After != D3D12_RESOURCE_STATE_TBD);
@@ -494,51 +490,15 @@ static inline bool IsTransitionNeeded(D3D12_RESOURCE_STATES Before, D3D12_RESOUR
 	return Before != After;
 }
 
-bool FD3D12ContextCommon::TransitionResource(FD3D12Resource* InResource, uint32 InSubresourceIndex, D3D12_RESOURCE_STATES InBeforeState, D3D12_RESOURCE_STATES InAfterState)
+void FD3D12ContextCommon::TransitionResource(FD3D12Resource* InResource, uint32 InSubresourceIndex, D3D12_RESOURCE_STATES InBeforeState, D3D12_RESOURCE_STATES InAfterState)
 {
 	check(InBeforeState != D3D12_RESOURCE_STATE_TBD);
 	check(InAfterState != D3D12_RESOURCE_STATE_TBD);
 
-	D3D12_RESOURCE_STATES BeforeState = InBeforeState;
-	D3D12_RESOURCE_STATES AfterState = InAfterState;
-
-	bool bRequireUAVBarrier = false;
-
-	// Require UAV barrier when before and after are UAV
-	if (BeforeState == D3D12_RESOURCE_STATE_UNORDERED_ACCESS && InAfterState == D3D12_RESOURCE_STATE_UNORDERED_ACCESS)
-	{
-		bRequireUAVBarrier = true;
-	}
-	// Special case for UAV access resources
-	else if (InResource->GetUAVAccessResource() && EnumHasAnyFlags(BeforeState | InAfterState, D3D12_RESOURCE_STATE_UNORDERED_ACCESS))
-	{
-		// We are issuing an Aliasing Barrier from or to GetUAVAccessResource
-		// while we are transitioning to a known resource state (D3D12_RESOURCE_STATE_PIXEL_SHADER_RESOURCE) the original resource since it can't go int UAV
-		// the transition into D3D12_RESOURCE_STATE_PIXEL_SHADER_RESOURCE is not needed but in this way the tracking of the resource will be consistent at the higher level
-
-		// inject an aliasing barrier
-		const bool bFromUAV = EnumHasAnyFlags(BeforeState, D3D12_RESOURCE_STATE_UNORDERED_ACCESS);
-		const bool bToUAV = EnumHasAnyFlags(InAfterState, D3D12_RESOURCE_STATE_UNORDERED_ACCESS);
-		check(bFromUAV != bToUAV);
-
-		if (bToUAV)
-		{
-			AddAliasingBarrier(InResource->GetResource(), InResource->GetUAVAccessResource());
-			AfterState = D3D12_RESOURCE_STATE_PIXEL_SHADER_RESOURCE; //original resource CAN'T go into D3D12_RESOURCE_STATE_UNORDERED_ACCESS so we transition it in a known state D3D12_RESOURCE_STATE_PIXEL_SHADER_RESOURCE
-		}
-		else
-		{
-			AddAliasingBarrier(InResource->GetUAVAccessResource(), InResource->GetResource());
-			BeforeState = D3D12_RESOURCE_STATE_PIXEL_SHADER_RESOURCE; //original resource CAN'T be into D3D12_RESOURCE_STATE_UNORDERED_ACCESS so it was transitioned to D3D12_RESOURCE_STATE_PIXEL_SHADER_RESOURCE
-		}
-	}
-
-	if (IsTransitionNeeded(BeforeState, AfterState, InResource))
+	if (IsTransitionNeeded(InBeforeState, InAfterState, InResource))
 	{
-		AddTransitionBarrier(InResource, BeforeState, AfterState, InSubresourceIndex);
+		AddTransitionBarrier(InResource, InBeforeState, InAfterState, InSubresourceIndex);
 	}
-
-	return bRequireUAVBarrier;
 }
 
 namespace D3D12RHI
diff --git a/Engine/Source/Runtime/D3D12RHI/Private/D3D12Commands.cpp b/Engine/Source/Runtime/D3D12RHI/Private/D3D12Commands.cpp
--- a/Engine/Source/Runtime/D3D12RHI/Private/D3D12Commands.cpp
+++ b/Engine/Source/Runtime/D3D12RHI/Private/D3D12Commands.cpp
@@ -572,10 +572,9 @@ void FD3D12CommandContext::HandleResourceTransitions(const FD3D12TransitionData*
 
 	for (const FRHITransitionInfo& Info : TransitionData->TransitionInfos)
 	{
-		const bool bUAVAccessAfter = EnumHasAnyFlags(Info.AccessAfter, ERHIAccess::UAVMask);
-
-		// The transition SRV->UAV needs a UAV barrier then to work correctly otherwise there is no synchronization at all
-		bUAVBarrier |= bUAVAccessAfter;
+		bUAVBarrier |=
+			EnumHasAnyFlags(Info.AccessBefore, ERHIAccess::UAVMask)
+			&& EnumHasAnyFlags(Info.AccessAfter, ERHIAccess::UAVMask);
 
 		// Very specific case that needs to be removed with EB
 		// a UAV -> SRVMAsk on the AsyncPipe get split in two: UAV->SRVCompute on Async and SRVCompute->SRVMask on Gfx
diff --git a/Engine/Source/Runtime/D3D12RHI/Private/D3D12Device.h b/Engine/Source/Runtime/D3D12RHI/Private/D3D12Device.h
--- a/Engine/Source/Runtime/D3D12RHI/Private/D3D12Device.h
+++ b/Engine/Source/Runtime/D3D12RHI/Private/D3D12Device.h
@@ -306,6 +306,15 @@ public:
 		void**                  ppCommandList
 	);
 
+#if !D3D12RHI_SUPPORTS_UNCOMPRESSED_UAV
+	void CreateUnorderedAccessViewAlias(
+		ID3D12Resource* InResource,
+		ID3D12Resource* InCounterResource,
+		const D3D12_RESOURCE_DESC& InAliasResourceDesc,
+		const D3D12_UNORDERED_ACCESS_VIEW_DESC& InAliasViewDesc,
+		D3D12_CPU_DESCRIPTOR_HANDLE InOfflineCpuHandle);
+#endif
+
 	TRefCountPtr<ID3D12CommandQueue> TileMappingQueue;
 	FD3D12Fence TileMappingFence;
 
diff --git a/Engine/Source/Runtime/D3D12RHI/Private/D3D12PoolAllocator.cpp b/Engine/Source/Runtime/D3D12RHI/Private/D3D12PoolAllocator.cpp
--- a/Engine/Source/Runtime/D3D12RHI/Private/D3D12PoolAllocator.cpp
+++ b/Engine/Source/Runtime/D3D12RHI/Private/D3D12PoolAllocator.cpp
@@ -416,8 +416,7 @@ void FD3D12PoolAllocator::AllocateResource(uint32 GPUIndex, D3D12_HEAP_TYPE InHe
 		Desc.Alignment = 0;
 
 		// If we are tracking all allocation data and allocating a standalone texture, then first create a heap so we can retrieve the GPU virtual address as well
-		// UAV Aliasing needs a Heap to create the aliased resource in.
-		if (InDesc.NeedsUAVAliasWorkarounds() || (Adapter->IsTrackingAllAllocations() && InDesc.Dimension != D3D12_RESOURCE_DIMENSION_BUFFER))
+		if (Adapter->IsTrackingAllAllocations() && (InDesc.Dimension != D3D12_RESOURCE_DIMENSION_BUFFER))
 		{
 			D3D12_HEAP_DESC HeapDesc = {};
 			HeapDesc.SizeInBytes = FMath::Max((uint64)D3D12_DEFAULT_RESOURCE_PLACEMENT_ALIGNMENT, InSize);
diff --git a/Engine/Source/Runtime/D3D12RHI/Private/D3D12Resources.cpp b/Engine/Source/Runtime/D3D12RHI/Private/D3D12Resources.cpp
--- a/Engine/Source/Runtime/D3D12RHI/Private/D3D12Resources.cpp
+++ b/Engine/Source/Runtime/D3D12RHI/Private/D3D12Resources.cpp
@@ -1102,29 +1102,19 @@ HRESULT FD3D12Adapter::CreateBuffer(const D3D12_HEAP_PROPERTIES& HeapProps,
 		ppOutResource, Name);
 }
 
-void FD3D12Adapter::CreateUAVAliasResource(D3D12_CLEAR_VALUE* ClearValuePtr, const TCHAR* DebugName, FD3D12ResourceLocation& Location)
+#if !D3D12RHI_SUPPORTS_UNCOMPRESSED_UAV
+void FD3D12Adapter::CreateUAVAliasResourceDesc(const FD3D12ResourceLocation& Location)
 {
 	FD3D12Resource* SourceResource = Location.GetResource();
 
 	const FD3D12ResourceDesc& SourceDesc = SourceResource->GetDesc();
-	const FD3D12Heap* const ResourceHeap = SourceResource->GetHeap();
 
 	const EPixelFormat SourceFormat = SourceDesc.PixelFormat;
 	const EPixelFormat AliasTextureFormat = SourceDesc.UAVPixelFormat;
 
-	if (ensure(ResourceHeap != nullptr) && ensure(SourceFormat != PF_Unknown) && SourceFormat != AliasTextureFormat)
+	if (ensure(SourceFormat != PF_Unknown) && (SourceFormat != AliasTextureFormat))
 	{
-		uint64 SourceOffset = Location.GetOffsetFromBaseOfResource();
-
-		if (Location.GetAllocatorType() == FD3D12ResourceLocation::AT_Pool && Location.GetType() == FD3D12ResourceLocation::ResourceLocationType::eSubAllocation)
-		{
-			FD3D12HeapAndOffset HeapAndOffset = Location.GetPoolAllocator()->GetBackingHeapAndAllocationOffsetInBytes(Location);
-			check(SourceOffset == 0);
-			check(ResourceHeap == HeapAndOffset.Heap);
-			SourceOffset = HeapAndOffset.Offset;
-		}
-
-		FD3D12ResourceDesc AliasTextureDesc = SourceDesc;
+		D3D12_RESOURCE_DESC AliasTextureDesc = SourceDesc;
 		AliasTextureDesc.Format = (DXGI_FORMAT)GPixelFormats[AliasTextureFormat].PlatformFormat;
 		AliasTextureDesc.Width = SourceDesc.Width / GPixelFormats[SourceFormat].BlockSizeX;
 		AliasTextureDesc.Height = SourceDesc.Height / GPixelFormats[SourceFormat].BlockSizeY;
@@ -1133,30 +1123,10 @@ void FD3D12Adapter::CreateUAVAliasResource(D3D12_CLEAR_VALUE* ClearValuePtr, con
 
 		EnumAddFlags(AliasTextureDesc.Flags, D3D12_RESOURCE_FLAG_ALLOW_UNORDERED_ACCESS);
 
-		AliasTextureDesc.UAVPixelFormat = PF_Unknown;
-
-		TRefCountPtr<ID3D12Resource> pAliasResource;
-		HRESULT AliasHR = GetD3DDevice()->CreatePlacedResource(
-			ResourceHeap->GetHeap(),
-			SourceOffset,
-			&AliasTextureDesc,
-			D3D12_RESOURCE_STATE_UNORDERED_ACCESS,
-			ClearValuePtr,
-			IID_PPV_ARGS(pAliasResource.GetInitReference()));
-
-		if (pAliasResource && DebugName)
-		{
-			TCHAR NameBuffer[512]{};
-			FCString::Snprintf(NameBuffer, UE_ARRAY_COUNT(NameBuffer), TEXT("%s UAVAlias"), DebugName);
-			SetName(pAliasResource, NameBuffer);
-		}
-
-		if (SUCCEEDED(AliasHR))
-		{
-			SourceResource->SetUAVAccessResource(pAliasResource);
-		}
+		SourceResource->SetUAVAccessResourceDesc(AliasTextureDesc);
 	}
 }
+#endif
 
 /////////////////////////////////////////////////////////////////////
 //	FD3D12 Resource Location
@@ -1571,10 +1541,6 @@ bool FD3D12ResourceLocation::OnAllocationMoved(FD3D12ContextArray const& Context
 		OtherResourceLocation->GPUVirtualAddress = GPUVirtualAddress;
 	}
 
-	// TODO: recreate the aliased UAV resource if we have one. For the time being we just check that we don't get here with such a resource,
-	// because defrag is disabled for the UAV pool in the FD3D12TextureAllocatorPool constructor.
-	check(!CurrentResource->GetDesc().NeedsUAVAliasWorkarounds());
-
 	// Notify all the dependent resources about the change
 	Owner->ResourceRenamed(Contexts);
 
diff --git a/Engine/Source/Runtime/D3D12RHI/Private/D3D12Resources.h b/Engine/Source/Runtime/D3D12RHI/Private/D3D12Resources.h
--- a/Engine/Source/Runtime/D3D12RHI/Private/D3D12Resources.h
+++ b/Engine/Source/Runtime/D3D12RHI/Private/D3D12Resources.h
@@ -189,10 +189,11 @@ class FD3D12Resource : public FThreadSafeRefCountedObject, public FD3D12DeviceCh
 	friend class FD3D12Buffer;
 private:
 	TRefCountPtr<ID3D12Resource> Resource;
-	// Since certain formats cannot be aliased in D3D12, we have to create a separate ID3D12Resource that aliases the
-	// resource's memory and use this separate resource to create the UAV.
-	// TODO: UE-116727: request better DX12 API support and clean this up.
-	TRefCountPtr<ID3D12Resource> UAVAccessResource;
+#if !D3D12RHI_SUPPORTS_UNCOMPRESSED_UAV
+	// Since certain formats cannot be aliased in D3D12, we create a placed UAV
+	// that aliases them. This resource description is used to create that UAV.
+	TUniquePtr<D3D12_RESOURCE_DESC> UAVAccessResourceDesc;
+#endif
 	TRefCountPtr<FD3D12Heap> Heap;
 
 	FD3D12ResidencyHandle* ResidencyHandle = nullptr; // Residency handle owned by this object
@@ -270,8 +271,18 @@ public:
 	operator ID3D12Resource&() { return *Resource; }
 
 	ID3D12Resource* GetResource() const { return Resource.GetReference(); }
-	ID3D12Resource* GetUAVAccessResource() const { return UAVAccessResource.GetReference(); }
-	void SetUAVAccessResource(ID3D12Resource* InUAVAccessResource) { UAVAccessResource = InUAVAccessResource; }
+
+#if !D3D12RHI_SUPPORTS_UNCOMPRESSED_UAV
+	const D3D12_RESOURCE_DESC* GetUAVAccessResourceDesc() const
+	{
+		return UAVAccessResourceDesc.Get();
+	}
+
+	void SetUAVAccessResourceDesc(const D3D12_RESOURCE_DESC& InUAVAccessResourceDesc)
+	{
+		UAVAccessResourceDesc = MakeUnique<D3D12_RESOURCE_DESC>(InUAVAccessResourceDesc);
+	}
+#endif
 
 	inline void* Map(const D3D12_RANGE* ReadRange = nullptr)
 	{
@@ -435,7 +446,7 @@ public:
 			bSRV(!EnumHasAnyFlags(Desc.Flags, D3D12_RESOURCE_FLAG_DENY_SHADER_RESOURCE)),
 			bDSV(EnumHasAnyFlags(Desc.Flags, D3D12_RESOURCE_FLAG_ALLOW_DEPTH_STENCIL)),
 			bRTV(EnumHasAnyFlags(Desc.Flags, D3D12_RESOURCE_FLAG_ALLOW_RENDER_TARGET)),
-			bUAV(EnumHasAnyFlags(Desc.Flags, D3D12_RESOURCE_FLAG_ALLOW_UNORDERED_ACCESS) || Desc.NeedsUAVAliasWorkarounds()),
+			bUAV(EnumHasAnyFlags(Desc.Flags, D3D12_RESOURCE_FLAG_ALLOW_UNORDERED_ACCESS)),
 			bWritable(bDSV || bRTV || bUAV),
 			bSRVOnly(bSRV && !bWritable),
 			bBuffer(Desc.Dimension == D3D12_RESOURCE_DIMENSION_BUFFER),
diff --git a/Engine/Source/Runtime/D3D12RHI/Private/D3D12Texture.cpp b/Engine/Source/Runtime/D3D12RHI/Private/D3D12Texture.cpp
--- a/Engine/Source/Runtime/D3D12RHI/Private/D3D12Texture.cpp
+++ b/Engine/Source/Runtime/D3D12RHI/Private/D3D12Texture.cpp
@@ -218,9 +218,8 @@ bool FD3D12Texture::CanBe4KAligned(const FD3D12ResourceDesc& Desc, EPixelFormat
 		return false;
 	}
 
-	// 4KB alignment is only available for read only textures
-	if (!EnumHasAnyFlags(Desc.Flags, D3D12_RESOURCE_FLAG_ALLOW_RENDER_TARGET | D3D12_RESOURCE_FLAG_ALLOW_DEPTH_STENCIL | D3D12_RESOURCE_FLAG_ALLOW_UNORDERED_ACCESS) &&
-		!Desc.NeedsUAVAliasWorkarounds() && // UAV aliased resources are secretly writable.
+	// 4KB alignment is only available for non-RT/DS textures
+	if (!EnumHasAnyFlags(Desc.Flags, D3D12_RESOURCE_FLAG_ALLOW_RENDER_TARGET | D3D12_RESOURCE_FLAG_ALLOW_DEPTH_STENCIL) &&
 		Desc.SampleDesc.Count == 1)
 	{
 		D3D12_TILE_SHAPE Tile = {};
@@ -353,7 +352,7 @@ FD3D12ResourceDesc FD3D12DynamicRHI::GetResourceDesc(const FRHITextureDesc& Text
 		ResourceDesc.UAVPixelFormat = GetBlockCompressedFormatUAVAliasFormat(TextureDesc.Format);
 	}
 
-	if (EnumHasAnyFlags(TextureDesc.Flags, TexCreate_UAV) && !ResourceDesc.NeedsUAVAliasWorkarounds())
+	if (EnumHasAnyFlags(TextureDesc.Flags, TexCreate_UAV))
 	{
 		ResourceDesc.Flags |= D3D12_RESOURCE_FLAG_ALLOW_UNORDERED_ACCESS;
 	}
@@ -730,10 +729,12 @@ FD3D12DynamicRHI::FCreateTextureInternalResult FD3D12DynamicRHI::CreateTextureIn
 
 		check(Location.IsValid());
 
+#if !D3D12RHI_SUPPORTS_UNCOMPRESSED_UAV
 		if (ResourceDesc.NeedsUAVAliasWorkarounds())
 		{
-			Adapter->CreateUAVAliasResource(ClearValuePtr, CreateDesc.DebugName, Location);
+			Adapter->CreateUAVAliasResourceDesc(Location);
 		}
+#endif
 
 		NewTexture->CreateViews(FirstLinkedObject);
 
diff --git a/Engine/Source/Runtime/D3D12RHI/Private/D3D12UAV.cpp b/Engine/Source/Runtime/D3D12RHI/Private/D3D12UAV.cpp
--- a/Engine/Source/Runtime/D3D12RHI/Private/D3D12UAV.cpp
+++ b/Engine/Source/Runtime/D3D12RHI/Private/D3D12UAV.cpp
@@ -84,20 +84,35 @@ void FD3D12UnorderedAccessView::ResourceRenamed(FD3D12ContextArray const& Contex
 
 void FD3D12UnorderedAccessView::UpdateDescriptor()
 {
-	ID3D12Resource* Resource = ResourceInfo.Resource->GetUAVAccessResource()
-		? ResourceInfo.Resource->GetUAVAccessResource()
-		: ResourceInfo.Resource->GetResource();
+	ID3D12Resource* Resource = ResourceInfo.Resource->GetResource();
 
 	ID3D12Resource* Counter = CounterResource
 		? CounterResource->GetResource()
 		: nullptr;
 
-	GetParentDevice()->GetDevice()->CreateUnorderedAccessView(
-		Resource, 
-		Counter, 
-		&D3DViewDesc,
-		OfflineCpuHandle
-	);
+#if !D3D12RHI_SUPPORTS_UNCOMPRESSED_UAV
+	const D3D12_RESOURCE_DESC* AliasResourceDesc =
+		ResourceInfo.Resource->GetUAVAccessResourceDesc();
+
+	if (AliasResourceDesc)
+	{
+		GetParentDevice()->CreateUnorderedAccessViewAlias(
+			Resource,
+			Counter,
+			*AliasResourceDesc,
+			D3DViewDesc,
+			OfflineCpuHandle);
+	}
+	else
+#endif
+	{
+		GetParentDevice()->GetDevice()->CreateUnorderedAccessView(
+			Resource, 
+			Counter, 
+			&D3DViewDesc,
+			OfflineCpuHandle
+		);
+	}
 
 	OfflineCpuHandle.IncrementVersion();
 }
