Allow DX12 Resource barrier flushing trough ID3D12DynamicRHI

diff --git a/Engine/Source/Runtime/D3D12RHI/Private/D3D12RHI.cpp b/Engine/Source/Runtime/D3D12RHI/Private/D3D12RHI.cpp
--- a/Engine/Source/Runtime/D3D12RHI/Private/D3D12RHI.cpp
+++ b/Engine/Source/Runtime/D3D12RHI/Private/D3D12RHI.cpp
@@ -784,6 +784,11 @@ void FD3D12DynamicRHI::RHIVerifyResult(ID3D12Device* Device, HRESULT Result, con
 	D3D12RHI::VerifyD3D12Result(Result, Code, Filename, Line, Device, Message);
 }
 
+void FD3D12DynamicRHI::RHIFlushResourceBarriers(FRHICommandListBase& RHICmdList, uint32 InGPUIndex)
+{
+	FD3D12CommandContext::Get(RHICmdList, InGPUIndex).FlushResourceBarriers();
+}
+
 void* FD3D12DynamicRHI::RHIGetNativeDevice()
 {
 	return (void*)GetAdapter().GetD3DDevice();
diff --git a/Engine/Source/Runtime/D3D12RHI/Private/D3D12RHIPrivate.h b/Engine/Source/Runtime/D3D12RHI/Private/D3D12RHIPrivate.h
--- a/Engine/Source/Runtime/D3D12RHI/Private/D3D12RHIPrivate.h
+++ b/Engine/Source/Runtime/D3D12RHI/Private/D3D12RHIPrivate.h
@@ -418,6 +418,7 @@ public:
 	virtual FRWLock* RHIGetSuspendedLock() { return nullptr; }
 
 	virtual void RHIRunOnQueue(ED3D12RHIRunOnQueueType QueueType, TFunction<void(ID3D12CommandQueue*)>&& CodeToRun, bool bWaitForSubmission) final override;
+	virtual void RHIFlushResourceBarriers(FRHICommandListBase& RHICmdList, uint32 InGPUIndex) final override;
 
 	virtual IRHITransientResourceAllocator* RHICreateTransientResourceAllocator() override;
 
diff --git a/Engine/Source/Runtime/D3D12RHI/Public/ID3D12DynamicRHI.h b/Engine/Source/Runtime/D3D12RHI/Public/ID3D12DynamicRHI.h
--- a/Engine/Source/Runtime/D3D12RHI/Public/ID3D12DynamicRHI.h
+++ b/Engine/Source/Runtime/D3D12RHI/Public/ID3D12DynamicRHI.h
@@ -57,6 +57,8 @@ struct ID3D12DynamicRHI : public FDynamicRHI
 	virtual void                  RHIVerifyResult(ID3D12Device* Device, HRESULT Result, const ANSICHAR* Code, const ANSICHAR* Filename, uint32 Line, FString Message = FString()) const = 0;
 
 	virtual void				  RHIRunOnQueue(ED3D12RHIRunOnQueueType QueueType, TFunction<void(ID3D12CommandQueue*)>&& CodeToRun, bool bWaitForSubmission) = 0;
+
+	virtual void				  RHIFlushResourceBarriers(FRHICommandListBase& RHICmdList, uint32 InGPUIndex) = 0;
 };
 
 inline bool IsRHID3D12()
