Release used PSOs and precached compute PSOs from the low-level D3D12 cache instead of always retaining them until application exit.

Also modify RHI API to take a compute pipeline state initializer instead of a pointer to a compute shader to allow to pass other information down (such as whether the PSO creation request comes from the precaching or file cache systems). This is similar to how the API works for graphics PSOs. The previous API is deprecated.

Before this change:
* PSOs used by the renderer would never be released by the low-level cache. This is problematic if the graphics driver expects the application to keep the number of live PSOs under a certain limit.
* Compute PSOs created by the precaching system would never be released by the low-level cache because the D3D12 RHI had no information on whether the PSO was coming from the renderer or the precaching system. This would cause a steady leak in CPU memory, which would become significant when compute-heavy systems like Nanite with compute materials were used.

#tests tested locally and by QA
#tests tested locally and by QA

diff --git a/Engine/Source/Runtime/Apple/MetalRHI/Private/MetalDynamicRHI.cpp b/Engine/Source/Runtime/Apple/MetalRHI/Private/MetalDynamicRHI.cpp
--- a/Engine/Source/Runtime/Apple/MetalRHI/Private/MetalDynamicRHI.cpp
+++ b/Engine/Source/Runtime/Apple/MetalRHI/Private/MetalDynamicRHI.cpp
@@ -102,10 +102,10 @@ FGraphicsPipelineStateRHIRef FMetalDynamicRHI::RHICreateGraphicsPipelineState(co
     return FGraphicsPipelineStateRHIRef(MoveTemp(State));
 }
 
-TRefCountPtr<FRHIComputePipelineState> FMetalDynamicRHI::RHICreateComputePipelineState(FRHIComputeShader* ComputeShader)
+TRefCountPtr<FRHIComputePipelineState> FMetalDynamicRHI::RHICreateComputePipelineState(const FComputePipelineStateInitializer& Initializer)
 {
-    MTL_SCOPED_AUTORELEASE_POOL;
-    return new FRHIComputePipelineState(ComputeShader);
+	MTL_SCOPED_AUTORELEASE_POOL;
+	return new FRHIComputePipelineState(Initializer.ComputeShader);
 }
 
 
diff --git a/Engine/Source/Runtime/Apple/MetalRHI/Private/MetalDynamicRHI.h b/Engine/Source/Runtime/Apple/MetalRHI/Private/MetalDynamicRHI.h
--- a/Engine/Source/Runtime/Apple/MetalRHI/Private/MetalDynamicRHI.h
+++ b/Engine/Source/Runtime/Apple/MetalRHI/Private/MetalDynamicRHI.h
@@ -116,7 +116,7 @@ public:
 #endif
 	virtual FBoundShaderStateRHIRef RHICreateBoundShaderState(FRHIVertexDeclaration* VertexDeclaration, FRHIVertexShader* VertexShader, FRHIPixelShader* PixelShader, FRHIGeometryShader* GeometryShader) final override;
 	virtual FGraphicsPipelineStateRHIRef RHICreateGraphicsPipelineState(const FGraphicsPipelineStateInitializer& Initializer) final override;
-	virtual TRefCountPtr<FRHIComputePipelineState> RHICreateComputePipelineState(FRHIComputeShader* ComputeShader) final override;
+	virtual TRefCountPtr<FRHIComputePipelineState> RHICreateComputePipelineState(const FComputePipelineStateInitializer& Initializer) final override;
 	virtual FUniformBufferRHIRef RHICreateUniformBuffer(const void* Contents, const FRHIUniformBufferLayout* Layout, EUniformBufferUsage Usage, EUniformBufferValidation Validation) final override;
 
 	[[nodiscard]] virtual FRHIBufferInitializer RHICreateBufferInitializer(FRHICommandListBase& RHICmdList, const FRHIBufferCreateDesc& CreateDesc) final override;
diff --git a/Engine/Source/Runtime/D3D12RHI/Private/D3D12PipelineState.h b/Engine/Source/Runtime/D3D12RHI/Private/D3D12PipelineState.h
--- a/Engine/Source/Runtime/D3D12RHI/Private/D3D12PipelineState.h
+++ b/Engine/Source/Runtime/D3D12RHI/Private/D3D12PipelineState.h
@@ -423,10 +423,11 @@ struct FD3D12GraphicsPipelineState : public FRHIGraphicsPipelineState, FD3D12Pip
 struct FD3D12ComputePipelineState : public FRHIComputePipelineState, FD3D12PipelineStateCommonData
 {
 	FD3D12ComputePipelineState() = delete;
-	FD3D12ComputePipelineState(FD3D12ComputeShader* InComputeShader, const FD3D12RootSignature* InRootSignature, FD3D12PipelineState* InPipelineState);
+	FD3D12ComputePipelineState(const FComputePipelineStateInitializer& InInitializer, const FD3D12RootSignature* InRootSignature, FD3D12PipelineState* InPipelineState);
 	~FD3D12ComputePipelineState();
 
-	bool bShaderNeedsGlobalConstantBuffer;
+	bool bPSOPrecache : 1;
+	bool bShaderNeedsGlobalConstantBuffer : 1;	
 };
 
 #if D3D12RHI_USE_HIGH_LEVEL_PSO_CACHE
@@ -531,6 +532,7 @@ protected:
 #endif
 public:
 	void RemoveFromLowLevelCache(FD3D12PipelineState* PipelineState, const FGraphicsPipelineStateInitializer& PipelineStateInitializer, const FD3D12RootSignature* RootSignature);
+	void RemoveFromLowLevelCache(FD3D12PipelineState* PipelineState, const FD3D12ComputeShader* ComputeShader, const FD3D12RootSignature* RootSignature);
 
 #if D3D12RHI_USE_HIGH_LEVEL_PSO_CACHE
 	FD3D12GraphicsPipelineState* FindInRuntimeCache(const FGraphicsPipelineStateInitializer& Initializer, uint32& OutHash);
@@ -539,8 +541,8 @@ public:
 
 	FD3D12ComputePipelineState* FindInRuntimeCache(const FD3D12ComputeShader* ComputeShader);
 #endif
-	FD3D12ComputePipelineState* FindInLoadedCache(FD3D12ComputeShader* ComputeShader, const FD3D12RootSignature* RootSignature, FD3D12ComputePipelineStateDesc& OutLowLevelDesc);
-	FD3D12ComputePipelineState* CreateAndAdd(FD3D12ComputeShader* ComputeShader, const FD3D12RootSignature* RootSignature, const FD3D12ComputePipelineStateDesc& LowLevelDesc);
+	FD3D12ComputePipelineState* FindInLoadedCache(const FComputePipelineStateInitializer& Initializer, const FD3D12RootSignature* RootSignature, FD3D12ComputePipelineStateDesc& OutLowLevelDesc);
+	FD3D12ComputePipelineState* CreateAndAdd(const FComputePipelineStateInitializer& Initializer, const FD3D12RootSignature* RootSignature, const FD3D12ComputePipelineStateDesc& LowLevelDesc);
 
 	static uint64 HashPSODesc(const FD3D12LowLevelGraphicsPipelineStateDesc& Desc);
 	static uint64 HashPSODesc(const FD3D12ComputePipelineStateDesc& Desc);
diff --git a/Engine/Source/Runtime/D3D12RHI/Private/D3D12Pipelinestate.cpp b/Engine/Source/Runtime/D3D12RHI/Private/D3D12Pipelinestate.cpp
--- a/Engine/Source/Runtime/D3D12RHI/Private/D3D12Pipelinestate.cpp
+++ b/Engine/Source/Runtime/D3D12RHI/Private/D3D12Pipelinestate.cpp
@@ -8,6 +8,7 @@
 #include "D3D12PipelineState.h"
 #include "D3D12RHIPrivate.h"
 #include "Hash/CityHash.h"
+#include "PipelineStateCache.h"
 
 static TAutoConsoleVariable<float> CVarPSOStallWarningThresholdInMs(
 	TEXT("D3D12.PSO.StallWarningThresholdInMs"),
@@ -16,14 +17,22 @@ static TAutoConsoleVariable<float> CVarPSOStallWarningThresholdInMs(
 	TEXT("Value is in milliseconds. (100 is the default)\n"),
 	ECVF_ReadOnly);
 
-int32 GPSOPrecacheKeepLowLevel = 0;
+bool GLowLevelCacheKeepPrecachedPSOs = false;
 static FAutoConsoleVariableRef CVarPSOPrecacheKeepLowLevel(
 	TEXT("D3D12.PSOPrecache.KeepLowLevel"),
-	GPSOPrecacheKeepLowLevel,
+	GLowLevelCacheKeepPrecachedPSOs,
 	TEXT("Keep in memory the d3d12 PSO blob for precached PSOs. Consumes more memory but reduces stalls.\n"),
 	ECVF_ReadOnly
 );
 
+bool GLowLevelCacheKeepUsedPSOs = false;
+static FAutoConsoleVariableRef CVarLowLevelCacheKeepUsedPSOs(
+	TEXT("D3D12.PSO.KeepUsedPSOsInLowLevelCache"),
+	GLowLevelCacheKeepUsedPSOs,
+	TEXT("Keep in memory the d3d12 PSO blob for used PSOs. Consumes more memory but reduces potential stalls.\n"),
+	ECVF_ReadOnly
+);
+
 /// @cond DOXYGEN_WARNINGS
 
 static void TranslateRenderTargetFormats(
@@ -313,8 +322,9 @@ bool FD3D12PipelineState::UsePSORefCounting()
 #if D3D12_USE_DERIVED_PSO
 	return false;
 #else
-	static const auto CVarPSOPrecaching = IConsoleManager::Get().FindConsoleVariable(TEXT("r.PSOPrecaching"));
-	return CVarPSOPrecaching && (CVarPSOPrecaching->GetInt() != 0);
+	bool bShouldRefCountForPSOPrecaching = PipelineStateCache::IsPSOPrecachingEnabled() && !GLowLevelCacheKeepPrecachedPSOs;
+	bool bShouldRefCountForUsedPSOs = !GLowLevelCacheKeepUsedPSOs;
+	return bShouldRefCountForPSOPrecaching || bShouldRefCountForUsedPSOs;
 #endif
 }
 
@@ -422,7 +432,9 @@ FD3D12GraphicsPipelineState::~FD3D12GraphicsPipelineState()
 		check(RefCount > 0);
         // precache PSO are here to avoid hitches at runtime when we want to create one that is actually used. We don't need to keep them
 		// around as this can add up to a lot of system memory
-		if (PipelineStateInitializer.bPSOPrecache && RefCount == 1 && GPSOPrecacheKeepLowLevel == 0)
+		bool bRemovePrecachedPSO = PipelineStateInitializer.bPSOPrecache && !GLowLevelCacheKeepPrecachedPSOs;
+		bool bRemoveUsedPSO = !PipelineStateInitializer.bPSOPrecache && !GLowLevelCacheKeepUsedPSOs;
+		if (RefCount == 1 && (bRemoveUsedPSO || bRemovePrecachedPSO))
 		{
 			FD3D12DynamicRHI* D3D12RHI = FD3D12DynamicRHI::GetD3DRHI();
 			FD3D12PipelineStateCache& PSOCache = D3D12RHI->GetAdapter().GetPSOCache();
@@ -448,11 +460,13 @@ FD3D12GraphicsPipelineState::~FD3D12GraphicsPipelineState()
 	}
 }
 
-FD3D12ComputePipelineState::FD3D12ComputePipelineState(FD3D12ComputeShader* InComputeShader, const FD3D12RootSignature* InRootSignature, FD3D12PipelineState* InPipelineState)
-	: FRHIComputePipelineState(InComputeShader)
+FD3D12ComputePipelineState::FD3D12ComputePipelineState(const FComputePipelineStateInitializer& InInitializer, const FD3D12RootSignature* InRootSignature, FD3D12PipelineState* InPipelineState)
+	: FRHIComputePipelineState(InInitializer.ComputeShader)
 	, FD3D12PipelineStateCommonData(InRootSignature, InPipelineState)
 {
-	bShaderNeedsGlobalConstantBuffer = InComputeShader && InComputeShader->UsesGlobalUniformBuffer();
+	FD3D12ComputeShader* D3D12ComputeShader = static_cast<FD3D12ComputeShader*>(InInitializer.ComputeShader);
+	bShaderNeedsGlobalConstantBuffer = D3D12ComputeShader && D3D12ComputeShader->UsesGlobalUniformBuffer();
+	bPSOPrecache = InInitializer.bPSOPrecache;
 }
 
 FD3D12ComputePipelineState::~FD3D12ComputePipelineState()
@@ -460,6 +474,20 @@ FD3D12ComputePipelineState::~FD3D12ComputePipelineState()
 	// At this point the object is not safe to use in the PSO cache.
 	// Currently, the PSO cache manages the lifetime but we could potentially
 	// stop doing an AddRef() and remove the PipelineState from any caches at this point.
+	if (FD3D12PipelineState::UsePSORefCounting() && PipelineState != nullptr)
+	{
+		uint32 RefCount = PipelineState->Release();
+		check(RefCount > 0);
+		bool bRemovePrecachedPSO = bPSOPrecache && !GLowLevelCacheKeepPrecachedPSOs;
+		bool bRemoveUsedPSO = !bPSOPrecache && !GLowLevelCacheKeepUsedPSOs;
+		if (RefCount == 1 && (bRemoveUsedPSO || bRemovePrecachedPSO))
+		{
+			FD3D12DynamicRHI* D3D12RHI = FD3D12DynamicRHI::GetD3DRHI();
+			FD3D12PipelineStateCache& PSOCache = D3D12RHI->GetAdapter().GetPSOCache();
+			// NB: it's possible that this remove will not do anything because another thread might have requested the same PSO since we entered the if
+			PSOCache.RemoveFromLowLevelCache(PipelineState, static_cast<FD3D12ComputeShader*>(GetComputeShader()), RootSignature);
+		}
+	}
 }
 
 void FD3D12PipelineStateCacheBase::CleanupPipelineStateCaches()
@@ -514,7 +542,17 @@ void FD3D12PipelineStateCacheBase::CleanupPipelineStateCaches()
 		for (auto Iter = ComputePipelineStateCache.CreateConstIterator(); Iter; ++Iter)
 		{
 			const FD3D12PipelineState* const PipelineState = Iter.Value();
-			delete PipelineState;
+			if (FD3D12PipelineState::UsePSORefCounting())
+			{
+				if (PipelineState)
+				{
+					PipelineState->Release();
+				}
+			}
+			else
+			{
+				delete PipelineState;
+			}
 		}
 		ComputePipelineStateCache.Empty();
 	}
@@ -600,6 +638,31 @@ void FD3D12PipelineStateCacheBase::RemoveFromLowLevelCache(FD3D12PipelineState*
 	}
 }
 
+void FD3D12PipelineStateCacheBase::RemoveFromLowLevelCache(FD3D12PipelineState* PipelineState, const FD3D12ComputeShader* ComputeShader, const FD3D12RootSignature* RootSignature)
+{
+	if (!FD3D12PipelineState::UsePSORefCounting())
+	{
+		checkNoEntry();
+		return;
+	}
+
+	FRWScopeLock Lock(ComputePipelineStateCacheMutex, FRWScopeLockType::SLT_Write);
+	// By the time we reach this function, it's possible another thread made a request for the same PSO
+	if (PipelineState->Release() == 0)
+	{
+		FD3D12ComputePipelineStateDesc LowLevelDesc = GetComputePipelineStateDesc(ComputeShader, RootSignature);
+		LowLevelDesc.Desc.NodeMask = FRHIGPUMask::All().GetNative();
+		LowLevelDesc.CombinedHash = FD3D12PipelineStateCacheBase::HashPSODesc(LowLevelDesc);
+		int32 ElementsRemoved = ComputePipelineStateCache.Remove(LowLevelDesc);
+		ensure(ElementsRemoved == 1);
+	}
+	else
+	{
+		// If another thread requested the pipeline state, we need to restore the refcount we just decremented
+		PipelineState->AddRef().CheckAtLeast(2);
+	}
+}
+
 void FD3D12PipelineStateCacheBase::AddToLowLevelCache(const FD3D12LowLevelGraphicsPipelineStateDesc& Desc, FD3D12PipelineState** OutPipelineState, const FPostCreateGraphicCallback& PostCreateCallback)
 {
 	check(Desc.CombinedHash != 0);
@@ -677,6 +740,10 @@ FD3D12PipelineState* FD3D12PipelineStateCacheBase::FindInLowLevelCache(const FD3
 		if (Found)
 		{
 			INC_DWORD_STAT(STAT_PSOComputeLowlevelCacheHit);
+			if (FD3D12PipelineState::UsePSORefCounting())
+			{
+				(*Found)->AddRef().CheckAtLeast(2);
+			}
 			return *Found;
 		}
 	}
@@ -711,6 +778,10 @@ void FD3D12PipelineStateCacheBase::AddToLowLevelCache(const FD3D12ComputePipelin
 		{
 			// This desc already exists.
 			*OutPipelineState = *PipelineState;
+			if (FD3D12PipelineState::UsePSORefCounting())
+			{
+				(*OutPipelineState)->AddRef();
+			}
 			return;
 		}
 
@@ -720,7 +791,19 @@ void FD3D12PipelineStateCacheBase::AddToLowLevelCache(const FD3D12ComputePipelin
 		FD3D12PipelineState* NewPipelineState = new FD3D12PipelineState(GetParentAdapter());
 		ComputePipelineStateCache.Add(Desc, NewPipelineState);
 
+		// This AddRef is for the low level cache.
+		if (FD3D12PipelineState::UsePSORefCounting())
+		{
+			NewPipelineState->AddRef();
+		}
+
 		*OutPipelineState = NewPipelineState;
+
+		// This AddRef is for the FD3D12ComputePipelineState requesting it.
+		if (FD3D12PipelineState::UsePSORefCounting())
+		{
+			(*OutPipelineState)->AddRef();
+		}
 	}
 
 	// Create the underlying PSO and then perform any other additional tasks like cleaning up/adding to caches, etc.
@@ -828,10 +911,10 @@ FD3D12ComputePipelineState* FD3D12PipelineStateCacheBase::FindInRuntimeCache(con
 }
 #endif
 
-FD3D12ComputePipelineState* FD3D12PipelineStateCacheBase::FindInLoadedCache(FD3D12ComputeShader* ComputeShader, const FD3D12RootSignature* RootSignature, FD3D12ComputePipelineStateDesc& OutLowLevelDesc)
+FD3D12ComputePipelineState* FD3D12PipelineStateCacheBase::FindInLoadedCache(const FComputePipelineStateInitializer& Initializer, const FD3D12RootSignature* RootSignature, FD3D12ComputePipelineStateDesc& OutLowLevelDesc)
 {
 	// TODO: For now PSOs will be created on every node of the LDA chain.
-	OutLowLevelDesc = GetComputePipelineStateDesc(ComputeShader, RootSignature);
+	OutLowLevelDesc = GetComputePipelineStateDesc(static_cast<FD3D12ComputeShader*>(Initializer.ComputeShader), RootSignature);
 	OutLowLevelDesc.Desc.NodeMask = FRHIGPUMask::All().GetNative();
 	OutLowLevelDesc.CombinedHash = FD3D12PipelineStateCacheBase::HashPSODesc(OutLowLevelDesc);
 
@@ -850,14 +933,14 @@ FD3D12ComputePipelineState* FD3D12PipelineStateCacheBase::FindInLoadedCache(FD3D
 #else
 	if (PipelineState && PipelineState->IsValid())
 	{
-		return new FD3D12ComputePipelineState(ComputeShader, RootSignature, PipelineState);
+		return new FD3D12ComputePipelineState(Initializer, RootSignature, PipelineState);
 	}
 
 	return nullptr;
 #endif
 }
 
-FD3D12ComputePipelineState* FD3D12PipelineStateCacheBase::CreateAndAdd(FD3D12ComputeShader* ComputeShader, const FD3D12RootSignature* RootSignature, const FD3D12ComputePipelineStateDesc& LowLevelDesc)
+FD3D12ComputePipelineState* FD3D12PipelineStateCacheBase::CreateAndAdd(const FComputePipelineStateInitializer& Initializer, const FD3D12RootSignature* RootSignature, const FD3D12ComputePipelineStateDesc& LowLevelDesc)
 {
 	FD3D12PipelineState* const PipelineState = CreateAndAddToLowLevelCache(LowLevelDesc);
 #if D3D12RHI_USE_HIGH_LEVEL_PSO_CACHE
@@ -866,7 +949,7 @@ FD3D12ComputePipelineState* FD3D12PipelineStateCacheBase::CreateAndAdd(FD3D12Com
 #else
 	if (PipelineState && PipelineState->IsValid())
 	{
-		return new FD3D12ComputePipelineState(ComputeShader, RootSignature, PipelineState);
+		return new FD3D12ComputePipelineState(Initializer, RootSignature, PipelineState);
 	}
 
 	return nullptr;
diff --git a/Engine/Source/Runtime/D3D12RHI/Private/D3D12RHIPrivate.h b/Engine/Source/Runtime/D3D12RHI/Private/D3D12RHIPrivate.h
--- a/Engine/Source/Runtime/D3D12RHI/Private/D3D12RHIPrivate.h
+++ b/Engine/Source/Runtime/D3D12RHI/Private/D3D12RHIPrivate.h
@@ -331,7 +331,7 @@ public:
     virtual void RHIUnlockStagingBuffer(FRHIStagingBuffer* StagingBuffer) final override;
 	virtual FBoundShaderStateRHIRef RHICreateBoundShaderState(FRHIVertexDeclaration* VertexDeclaration, FRHIVertexShader* VertexShader, FRHIPixelShader* PixelShader, FRHIGeometryShader* GeometryShader) final override;
 	virtual FGraphicsPipelineStateRHIRef RHICreateGraphicsPipelineState(const FGraphicsPipelineStateInitializer& Initializer) final override;
-	virtual FComputePipelineStateRHIRef RHICreateComputePipelineState(FRHIComputeShader* ComputeShader) final override;
+	virtual FComputePipelineStateRHIRef RHICreateComputePipelineState(const FComputePipelineStateInitializer& Initializer) final override;
 	virtual FWorkGraphPipelineStateRHIRef RHICreateWorkGraphPipelineState(const FWorkGraphPipelineStateInitializer& Initializer) final override;
 	virtual void RHICreateTransition(FRHITransition* Transition, const FRHITransitionCreateInfo& CreateInfo) final override;
 	virtual void RHIReleaseTransition(FRHITransition* Transition) final override;
diff --git a/Engine/Source/Runtime/D3D12RHI/Private/D3D12State.cpp b/Engine/Source/Runtime/D3D12RHI/Private/D3D12State.cpp
--- a/Engine/Source/Runtime/D3D12RHI/Private/D3D12State.cpp
+++ b/Engine/Source/Runtime/D3D12RHI/Private/D3D12State.cpp
@@ -758,18 +758,18 @@ FGraphicsPipelineStateRHIRef FD3D12DynamicRHI::RHICreateGraphicsPipelineState(co
 #endif
 }
 
-TRefCountPtr<FRHIComputePipelineState> FD3D12DynamicRHI::RHICreateComputePipelineState(FRHIComputeShader* ComputeShaderRHI)
+TRefCountPtr<FRHIComputePipelineState> FD3D12DynamicRHI::RHICreateComputePipelineState(const FComputePipelineStateInitializer& Initializer)
 {
 	SCOPE_CYCLE_COUNTER(STAT_PSOComputeFindOrCreateTime);
 
-	check(ComputeShaderRHI);
+	check(Initializer.ComputeShader);
 	FD3D12PipelineStateCache& PSOCache = GetAdapter().GetPSOCache();
-	FD3D12ComputeShader* ComputeShader = FD3D12DynamicRHI::ResourceCast(ComputeShaderRHI);
+	FD3D12ComputeShader* ComputeShader = FD3D12DynamicRHI::ResourceCast(Initializer.ComputeShader);
 
 	// First try to find the PSO based on Compute Shader pointer.
 	FD3D12ComputePipelineState* Found;
 #if D3D12RHI_USE_HIGH_LEVEL_PSO_CACHE
-	Found = PSOCache.FindInRuntimeCache(ComputeShader);
+	Found = PSOCache.FindInRuntimeCache(Initializer.ComputeShader);
 	if (Found)
 	{
 		return Found;
@@ -786,7 +786,7 @@ TRefCountPtr<FRHIComputePipelineState> FD3D12DynamicRHI::RHICreateComputePipelin
 
 	// Next try to find the PSO based on the hash of its desc.
 	FD3D12ComputePipelineStateDesc LowLevelDesc;
-	Found = PSOCache.FindInLoadedCache(ComputeShader, RootSignature, LowLevelDesc);
+	Found = PSOCache.FindInLoadedCache(Initializer, RootSignature, LowLevelDesc);
 	if (Found)
 	{
 		return Found;
@@ -795,7 +795,7 @@ TRefCountPtr<FRHIComputePipelineState> FD3D12DynamicRHI::RHICreateComputePipelin
 	TRACE_CPUPROFILER_EVENT_SCOPE(FD3D12DynamicRHI::RHICreateComputePipelineState);
 
 	// We need to actually create a PSO.
-	return PSOCache.CreateAndAdd(ComputeShader, RootSignature, LowLevelDesc);
+	return PSOCache.CreateAndAdd(Initializer, RootSignature, LowLevelDesc);
 }
 
 FD3D12SamplerState::FD3D12SamplerState(FD3D12Device* InParent, const D3D12_SAMPLER_DESC& Desc, uint16 SamplerID, FD3D12SamplerState* FirstLinkedObject)
diff --git a/Engine/Source/Runtime/RHI/Private/PipelineStateCache.cpp b/Engine/Source/Runtime/RHI/Private/PipelineStateCache.cpp
--- a/Engine/Source/Runtime/RHI/Private/PipelineStateCache.cpp
+++ b/Engine/Source/Runtime/RHI/Private/PipelineStateCache.cpp
@@ -166,7 +166,7 @@ static TAutoConsoleVariable<int32> CVarPSOEvictionTime(
 	TEXT("r.pso.evictiontime"),
 	60,
 	TEXT("Time between checks to remove stale objects from the cache. 0 = no eviction (which may eventually OOM...)"),
-	ECVF_ReadOnly | ECVF_RenderThreadSafe
+	ECVF_RenderThreadSafe
 );
 
 static TAutoConsoleVariable<int32> CVarPSORuntimeCreationHitchThreshold(
@@ -3079,7 +3079,10 @@ public:
 				FComputePipelineState* ComputePipeline = static_cast<FComputePipelineState*>(Pipeline);
 
 				uint64 StartTime = FPlatformTime::Cycles64();
-				ComputePipeline->RHIPipeline = RHICreateComputePipelineState(ComputePipeline->ComputeShader);
+				FComputePipelineStateInitializer ComputeInitializer = { ComputePipeline->ComputeShader, 0 };
+				ComputeInitializer.bFromPSOFileCache = Initializer.bFromPSOFileCache;
+				ComputeInitializer.bPSOPrecache = Initializer.bPSOPrecache;
+				ComputePipeline->RHIPipeline = RHICreateComputePipelineState(ComputeInitializer);
 				CheckAndUpdateHitchCountStat(FPSOPrecacheRequestID::EType::Compute, !IsPrecachedPSO(Initializer), PSOCompilationDebugData, StartTime, PSOPrecacheResult);
 
 				if (!ComputePipeline->RHIPipeline)
@@ -3393,7 +3396,9 @@ static void InternalCreateComputePipelineState(FRHIComputeShader* ComputeShader,
 
 		check(GraphEvent == nullptr);
 		uint64 StartTime = FPlatformTime::Cycles64();
-		CachedState->RHIPipeline = RHICreateComputePipelineState(ComputeShader);
+		FComputePipelineStateInitializer ComputeInitializerRHI = { ComputeShader, 0 };
+		ComputeInitializerRHI.bFromPSOFileCache = bFromPSOFileCache;
+		CachedState->RHIPipeline = RHICreateComputePipelineState(ComputeInitializerRHI);
 		CheckAndUpdateHitchCountStat(FPSOPrecacheRequestID::EType::Compute, !bFromPSOFileCache, PSOCompilationDebugData, StartTime, PSOPrecacheResult);
 
 		if (!CachedState->RHIPipeline)
@@ -4181,7 +4186,9 @@ void FPrecacheComputePipelineCache::OnNewPipelineStateCreated(const FPrecacheCom
 	else
 	{
 		check(CachedState->GetCompletionEvent() == nullptr);
-		CachedState->RHIPipeline = RHICreateComputePipelineState(CachedState->ComputeShader);
+		FComputePipelineStateInitializer ComputeInitializerRHI = { CachedState->ComputeShader, 0 };
+		ComputeInitializerRHI.bPSOPrecache = true;
+		CachedState->RHIPipeline = RHICreateComputePipelineState(ComputeInitializerRHI);
 		GPrecacheComputePipelineCache->PrecacheFinished(ComputeInitializer, CachedState->RHIPipeline != nullptr);
 	}
 }
diff --git a/Engine/Source/Runtime/RHI/Public/DynamicRHI.h b/Engine/Source/Runtime/RHI/Public/DynamicRHI.h
--- a/Engine/Source/Runtime/RHI/Public/DynamicRHI.h
+++ b/Engine/Source/Runtime/RHI/Public/DynamicRHI.h
@@ -414,16 +414,16 @@ public:
 	virtual FGraphicsPipelineStateRHIRef RHICreateGraphicsPipelineState(const FGraphicsPipelineStateInitializer& Initializer) = 0;
 
 	// FlushType: Thread safe
-	virtual FComputePipelineStateRHIRef RHICreateComputePipelineState(FRHIComputeShader * ComputeShader) = 0;
+	virtual FComputePipelineStateRHIRef RHICreateComputePipelineState(const FComputePipelineStateInitializer& Initializer) = 0;
 
 	virtual FGraphicsPipelineStateRHIRef RHICreateGraphicsPipelineState(const FGraphicsPipelineStateInitializer& Initializer, FRHIPipelineBinaryLibrary* PipelineBinary)
 	{
 		return RHICreateGraphicsPipelineState(Initializer);
 	}
 
-	virtual FComputePipelineStateRHIRef RHICreateComputePipelineState(FRHIComputeShader* ComputeShader, FRHIPipelineBinaryLibrary* PipelineBinary)
+	virtual FComputePipelineStateRHIRef RHICreateComputePipelineState(const FComputePipelineStateInitializer& Initializer, FRHIPipelineBinaryLibrary* PipelineBinary)
 	{
-		return RHICreateComputePipelineState(ComputeShader);
+		return RHICreateComputePipelineState(Initializer);
 	}
 
 	virtual FWorkGraphPipelineStateRHIRef RHICreateWorkGraphPipelineState(const FWorkGraphPipelineStateInitializer& Initializer)
@@ -1069,9 +1069,9 @@ public:
 		return new FRHIGraphicsPipelineStateFallBack(Initializer);
 	}
 
-	virtual FComputePipelineStateRHIRef RHICreateComputePipelineState(FRHIComputeShader* ComputeShader) override
+	virtual FComputePipelineStateRHIRef RHICreateComputePipelineState(const FComputePipelineStateInitializer& Initializer) override
 	{
-		return new FRHIComputePipelineStateFallback(ComputeShader);
+		return new FRHIComputePipelineStateFallback(Initializer.ComputeShader);
 	}
 };
 
@@ -1195,10 +1195,17 @@ FORCEINLINE FVertexDeclarationRHIRef RHICreateVertexDeclaration(const FVertexDec
 	return GDynamicRHI->RHICreateVertexDeclaration(Elements);
 }
 
+UE_DEPRECATED(5.6, "Creating a compute pipeline state with a pointer to a FRHIComputeShader is deprecated, please pass a FComputePipelineStateInitializer instead.")
 FORCEINLINE FComputePipelineStateRHIRef RHICreateComputePipelineState(FRHIComputeShader* ComputeShader)
+{
+	FComputePipelineStateInitializer Initializer = { ComputeShader, 0 };
+	return GDynamicRHI->RHICreateComputePipelineState(Initializer);
+}
+
+FORCEINLINE FComputePipelineStateRHIRef RHICreateComputePipelineState(const FComputePipelineStateInitializer& Initializer)
 {
 	LLM_SCOPE(ELLMTag::Shaders);
-	return GDynamicRHI->RHICreateComputePipelineState(ComputeShader);
+	return GDynamicRHI->RHICreateComputePipelineState(Initializer);
 }
 
 FORCEINLINE TRefCountPtr<FRHIWorkGraphPipelineState> RHICreateWorkGraphPipelineState(const FWorkGraphPipelineStateInitializer& Initializer)
diff --git a/Engine/Source/Runtime/RHI/Public/RHIResources.h b/Engine/Source/Runtime/RHI/Public/RHIResources.h
--- a/Engine/Source/Runtime/RHI/Public/RHIResources.h
+++ b/Engine/Source/Runtime/RHI/Public/RHIResources.h
@@ -1117,6 +1117,8 @@ public:
 
 	inline void SetValid(bool InIsValid) { bIsValid = InIsValid; }
 	inline bool IsValid() const { return bIsValid; }
+	virtual void MarkUsed() { bUsed = true; }
+	bool IsUsed() { return bUsed; }
 
 	FORCEINLINE FRHIComputeShader* GetComputeShader() const
 	{
@@ -1128,6 +1130,7 @@ protected:
 
 private:
 	bool bIsValid = true;
+	bool bUsed = false;
 };
 
 class FRHIWorkGraphPipelineState : public FRHIResource
@@ -4746,6 +4749,47 @@ public:
 	uint64							StatePrecachePSOHash;
 };
 
+class FComputePipelineStateInitializer
+{
+public:
+
+	FComputePipelineStateInitializer()
+		: ComputeShader(nullptr)
+		, Flags(0)
+	{
+	}
+
+	FComputePipelineStateInitializer(
+		FRHIComputeShader* InComputeShader,
+		uint8 InFlags)
+		: ComputeShader(InComputeShader)
+		, Flags(InFlags)
+	{
+	}
+
+	bool operator==(const FComputePipelineStateInitializer& Other) const
+	{
+		return ComputeShader == Other.ComputeShader;
+	}
+
+	FRHIComputeShader* ComputeShader;
+
+	// Note: these flags do NOT affect compilation of this PSO.
+	// The resulting object is invariant with respect to whatever is set here, they are
+	// behavior hints.
+	// They do not participate in equality comparisons or hashing.
+	union
+	{
+		struct
+		{
+			uint8	Reserved : 6;
+			uint8	bPSOPrecache : 1;
+			uint8	bFromPSOFileCache : 1;
+		};
+		uint8		Flags;
+	};
+};
+
 /** Helper for fast compute of hash for a shader table. */
 template<typename TShaderType>
 inline uint64 ComputeShaderTableHash(const TArrayView<TShaderType*>& ShaderTable, uint64 InitialHash = 5699878132332235837ull)
diff --git a/Engine/Source/Runtime/RHI/Public/RHIValidation.h b/Engine/Source/Runtime/RHI/Public/RHIValidation.h
--- a/Engine/Source/Runtime/RHI/Public/RHIValidation.h
+++ b/Engine/Source/Runtime/RHI/Public/RHIValidation.h
@@ -378,9 +378,9 @@ public:
 		return PSO;
 	}
 
-	virtual TRefCountPtr<FRHIComputePipelineState> RHICreateComputePipelineState(FRHIComputeShader* ComputeShader) override final
+	virtual TRefCountPtr<FRHIComputePipelineState> RHICreateComputePipelineState(const FComputePipelineStateInitializer& Initializer) override final
 	{
-		return RHI->RHICreateComputePipelineState(ComputeShader);
+		return RHI->RHICreateComputePipelineState(Initializer);
 	}
 
 	virtual FGraphicsPipelineStateRHIRef RHICreateGraphicsPipelineState(const FGraphicsPipelineStateInitializer& Initializer, FRHIPipelineBinaryLibrary* PipelineBinary) override final
@@ -394,9 +394,9 @@ public:
 		return PSO;
 	}
 
-	virtual TRefCountPtr<FRHIComputePipelineState> RHICreateComputePipelineState(FRHIComputeShader* ComputeShader, FRHIPipelineBinaryLibrary* PipelineBinary) override final
+	virtual TRefCountPtr<FRHIComputePipelineState> RHICreateComputePipelineState(const FComputePipelineStateInitializer& Initializer, FRHIPipelineBinaryLibrary* PipelineBinary) override final
 	{
-		return RHI->RHICreateComputePipelineState(ComputeShader, PipelineBinary);
+		return RHI->RHICreateComputePipelineState(Initializer, PipelineBinary);
 	}
 
 	/**
diff --git a/Engine/Source/Runtime/VulkanRHI/Private/VulkanDynamicRHI.h b/Engine/Source/Runtime/VulkanRHI/Private/VulkanDynamicRHI.h
--- a/Engine/Source/Runtime/VulkanRHI/Private/VulkanDynamicRHI.h
+++ b/Engine/Source/Runtime/VulkanRHI/Private/VulkanDynamicRHI.h
@@ -148,7 +148,7 @@ public:
 	virtual FStagingBufferRHIRef RHICreateStagingBuffer() final override;
 	virtual FBoundShaderStateRHIRef RHICreateBoundShaderState(FRHIVertexDeclaration* VertexDeclaration, FRHIVertexShader* VertexShader, FRHIPixelShader* PixelShader, FRHIGeometryShader* GeometryShader) final override;
 	virtual FGraphicsPipelineStateRHIRef RHICreateGraphicsPipelineState(const FGraphicsPipelineStateInitializer& Initializer) final override;
-	virtual FComputePipelineStateRHIRef RHICreateComputePipelineState(FRHIComputeShader* ComputeShader) final override;
+	virtual FComputePipelineStateRHIRef RHICreateComputePipelineState(const FComputePipelineStateInitializer& Initializer) final override;
 	virtual FUniformBufferRHIRef RHICreateUniformBuffer(const void* Contents, const FRHIUniformBufferLayout* Layout, EUniformBufferUsage Usage, EUniformBufferValidation Validation) final override;
 	virtual void RHIUpdateUniformBuffer(FRHICommandListBase& RHICmdList, FRHIUniformBuffer* UniformBufferRHI, const void* Contents) final override;
 
diff --git a/Engine/Source/Runtime/VulkanRHI/Private/VulkanPipeline.cpp b/Engine/Source/Runtime/VulkanRHI/Private/VulkanPipeline.cpp
--- a/Engine/Source/Runtime/VulkanRHI/Private/VulkanPipeline.cpp
+++ b/Engine/Source/Runtime/VulkanRHI/Private/VulkanPipeline.cpp
@@ -2243,7 +2243,7 @@ FVulkanComputePipeline* FVulkanPipelineStateCacheManager::RHICreateComputePipeli
 	return Device->GetPipelineStateCache()->GetOrCreateComputePipeline(ComputeShader);
 }
 
-FComputePipelineStateRHIRef FVulkanDynamicRHI::RHICreateComputePipelineState(FRHIComputeShader* ComputeShader)
+FComputePipelineStateRHIRef FVulkanDynamicRHI::RHICreateComputePipelineState(const FComputePipelineStateInitializer& Initializer)
 {
 #if VULKAN_ENABLE_AGGRESSIVE_STATS
 	SCOPE_CYCLE_COUNTER(STAT_VulkanGetOrCreatePipeline);
@@ -2251,7 +2251,7 @@ FComputePipelineStateRHIRef FVulkanDynamicRHI::RHICreateComputePipelineState(FRH
 	QUICK_SCOPE_CYCLE_COUNTER(STAT_Vulkan_RHICreateComputePipelineState);
 	LLM_SCOPE_VULKAN(ELLMTagVulkan::VulkanShaders);
 
-	return Device->PipelineStateCache->RHICreateComputePipelineState(ComputeShader);
+	return Device->PipelineStateCache->RHICreateComputePipelineState(Initializer.ComputeShader);
 }
 
 FVulkanComputePipeline* FVulkanPipelineStateCacheManager::GetOrCreateComputePipeline(FVulkanComputeShader* ComputeShader)
