[HWRT] Refactor BuildAccelerationStructure(...) in D3D12RHI.
- merge two loops over ReferencedGeometries.
- moved some code around to be more clear and removed outdated comments.

diff --git a/Engine/Source/Runtime/D3D12RHI/Private/D3D12RayTracing.cpp b/Engine/Source/Runtime/D3D12RHI/Private/D3D12RayTracing.cpp
--- a/Engine/Source/Runtime/D3D12RHI/Private/D3D12RayTracing.cpp
+++ b/Engine/Source/Runtime/D3D12RHI/Private/D3D12RayTracing.cpp
@@ -4437,105 +4437,95 @@ void BuildAccelerationStructure(FD3D12CommandContext& CommandContext,
 		}
 	}
 
-	{
-		// - Set up acceleration structure pointers and make them resident.
+	// Make necessary resources resident
 
-		CommandContext.UpdateResidency(InstanceBuffer->GetResource());
+	TRefCountPtr<FD3D12Buffer>& AccelerationStructureBuffer = Scene.AccelerationStructureBuffers[GPUIndex];
+	checkf(AccelerationStructureBuffer.IsValid(),
+		TEXT("Acceleration structure buffer must be set for this scene using RHIBindAccelerationStructureMemory() before build command is issued."));
 
-		{
-			TArray<const FD3D12Resource*>& ResourcesToMakeResidentForThisGPU = Scene.ResourcesToMakeResident[GPUIndex];
+	CommandContext.UpdateResidency(AccelerationStructureBuffer->GetResource());
+	CommandContext.UpdateResidency(InstanceBuffer->GetResource());
+	CommandContext.UpdateResidency(ScratchBuffer->GetResource());
 
-			ResourcesToMakeResidentForThisGPU.Reset(0);
+	{
+		// at the same time also gather resources that need to be resident when using TLAS
 
-			Experimental::TSherwoodSet<FD3D12ResidencyHandle*> UniqueResidencyHandles;
+		TArray<const FD3D12Resource*>& ResourcesToMakeResidentForThisGPU = Scene.ResourcesToMakeResident[GPUIndex];
 
-			auto AddResidencyHandleForResource = [&UniqueResidencyHandles, &ResourcesToMakeResidentForThisGPU] (FD3D12Resource* Resource)
-			{
-			#if ENABLE_RESIDENCY_MANAGEMENT
+		ResourcesToMakeResidentForThisGPU.Reset(0);
 
-				bool bShouldTrackResidency = false;
+		Experimental::TSherwoodSet<FD3D12ResidencyHandle*> UniqueResidencyHandles;
 
-				if (Resource->NeedsDeferredResidencyUpdate())
-				{
-					// Resources whose residency handles might change dynamically must always be tracked
-					bShouldTrackResidency = true;
-				}
-				else
+		auto AddResidencyHandleForResource = [&UniqueResidencyHandles, &ResourcesToMakeResidentForThisGPU] (FD3D12Resource* Resource)
+		{
+#if ENABLE_RESIDENCY_MANAGEMENT
+
+			bool bShouldTrackResidency = false;
+
+			if (Resource->NeedsDeferredResidencyUpdate())
+			{
+				// Resources whose residency handles might change dynamically must always be tracked
+				bShouldTrackResidency = true;
+			}
+			else
+			{
+				// Resources that share *all* residency handles with what's already tracked don't need to be added to be tracked separately
+				for (FD3D12ResidencyHandle* ResidencyHandle : Resource->GetResidencyHandles())
 				{
-					// Resources that share *all* residency handles with what's already tracked don't need to be added to be tracked separately
-					for (FD3D12ResidencyHandle* ResidencyHandle : Resource->GetResidencyHandles())
+					if (D3DX12Residency::IsInitialized(ResidencyHandle))
 					{
-						if (D3DX12Residency::IsInitialized(ResidencyHandle))
+						bool bIsAlreadyInSet = false;
+						UniqueResidencyHandles.Add(ResidencyHandle, &bIsAlreadyInSet);
+						if (!bIsAlreadyInSet)
 						{
-							bool bIsAlreadyInSet = false;
-							UniqueResidencyHandles.Add(ResidencyHandle, &bIsAlreadyInSet);
-							if (!bIsAlreadyInSet)
-							{
-								bShouldTrackResidency = true;
-							}
+							bShouldTrackResidency = true;
 						}
 					}
 				}
+			}
 
-				if (bShouldTrackResidency)
-				{
-					ResourcesToMakeResidentForThisGPU.Add(Resource);
-				}
+			if (bShouldTrackResidency)
+			{
+				ResourcesToMakeResidentForThisGPU.Add(Resource);
+			}
 
-			#endif // ENABLE_RESIDENCY_MANAGEMENT
-			};
+#endif // ENABLE_RESIDENCY_MANAGEMENT
+		};
 
-			const int32 NumReferencedGeometries = Scene.ReferencedGeometries.Num();
-			for (int32 Index = 0; Index < NumReferencedGeometries; ++Index)
-			{
-				FD3D12RayTracingGeometry* Geometry = FD3D12DynamicRHI::ResourceCast(Scene.ReferencedGeometries[Index].GetReference());
+		const int32 NumReferencedGeometries = Scene.ReferencedGeometries.Num();
+		for (int32 Index = 0; Index < NumReferencedGeometries; ++Index)
+		{
+			FD3D12RayTracingGeometry* Geometry = FD3D12DynamicRHI::ResourceCast(Scene.ReferencedGeometries[Index].GetReference());
 
-				checkf(!Geometry->IsDirty(CommandContext.GetGPUIndex()),
-					TEXT("Acceleration structures for all geometries must be built before building the top level acceleration structure for the scene."));
+			checkf(!Geometry->IsDirty(CommandContext.GetGPUIndex()),
+				TEXT("Acceleration structures for all geometries must be built before building the top level acceleration structure for the scene."));
 
-				AddResidencyHandleForResource(Geometry->AccelerationStructureBuffers[GPUIndex]->GetResource());
+			CommandContext.UpdateResidency(Geometry->AccelerationStructureBuffers[GPUIndex]->ResourceLocation.GetResource());
 
-				if (GRHIGlobals.RayTracing.SupportsShaders || GRHIGlobals.RayTracing.RequiresInlineRayTracingSBT)
-				{
-					checkf(Geometry->BuffersValid(CommandContext.GetGPUIndex()),
-						TEXT("Index & vertex buffers for all geometries must be valid (streamed in) when adding geometry to the top level acceleration structure for the scene"));
+			AddResidencyHandleForResource(Geometry->AccelerationStructureBuffers[GPUIndex]->GetResource());
 
-					if (Geometry->Initializer.IndexBuffer)
-					{
-						FD3D12Buffer* IndexBuffer = CommandContext.RetrieveObject<FD3D12Buffer>(Geometry->Initializer.IndexBuffer.GetReference());
-						AddResidencyHandleForResource(IndexBuffer->GetResource());
-					}
+			if (GRHIGlobals.RayTracing.SupportsShaders || GRHIGlobals.RayTracing.RequiresInlineRayTracingSBT)
+			{
+				checkf(Geometry->BuffersValid(CommandContext.GetGPUIndex()),
+					TEXT("Index & vertex buffers for all geometries must be valid (streamed in) when adding geometry to the top level acceleration structure for the scene"));
+
+				if (Geometry->Initializer.IndexBuffer)
+				{
+					FD3D12Buffer* IndexBuffer = CommandContext.RetrieveObject<FD3D12Buffer>(Geometry->Initializer.IndexBuffer.GetReference());
+					AddResidencyHandleForResource(IndexBuffer->GetResource());
+				}
 
-					for (const FRayTracingGeometrySegment& Segment : Geometry->Initializer.Segments)
+				for (const FRayTracingGeometrySegment& Segment : Geometry->Initializer.Segments)
+				{
+					if (Segment.VertexBuffer)
 					{
-						if (Segment.VertexBuffer)
-						{
-							FD3D12Buffer* VertexBuffer = CommandContext.RetrieveObject<FD3D12Buffer>(Segment.VertexBuffer.GetReference());
-							AddResidencyHandleForResource(VertexBuffer->GetResource());
-						}
+						FD3D12Buffer* VertexBuffer = CommandContext.RetrieveObject<FD3D12Buffer>(Segment.VertexBuffer.GetReference());
+						AddResidencyHandleForResource(VertexBuffer->GetResource());
 					}
 				}
 			}
 		}
 	}
-	
-	// Build the actual acceleration structure
-
-	const int32 NumReferencedGeometries = Scene.ReferencedGeometries.Num();
-	for (int32 Index = 0; Index < NumReferencedGeometries; ++Index)
-	{
-		FD3D12RayTracingGeometry* Geometry = FD3D12DynamicRHI::ResourceCast(Scene.ReferencedGeometries[Index].GetReference());
-		CommandContext.UpdateResidency(Geometry->AccelerationStructureBuffers[GPUIndex]->ResourceLocation.GetResource());
-	}
-
-	TRefCountPtr<FD3D12Buffer>& AccelerationStructureBuffer = Scene.AccelerationStructureBuffers[GPUIndex];
-	checkf(AccelerationStructureBuffer.IsValid(), 
-		TEXT("Acceleration structure buffer must be set for this scene using RHIBindAccelerationStructureMemory() before build command is issued."));
-
-	CommandContext.UpdateResidency(AccelerationStructureBuffer->GetResource());
-	CommandContext.UpdateResidency(ScratchBuffer->GetResource());
-
-	// Enqueue transition to UAV/SRV
 
 	if (ShouldRunRayTracingGPUValidation())
 	{
@@ -4552,13 +4542,7 @@ void BuildAccelerationStructure(FD3D12CommandContext& CommandContext,
 			InstanceBuffer, InstanceBufferOffset, InstanceBufferStride);
 	}
 
-	// UAV barrier is used here to ensure that all bottom level acceleration structures are built
-	CommandContext.AddUAVBarrier();
-	CommandContext.FlushResourceBarriers();
-
-	TArray<D3D12_BUILD_RAYTRACING_ACCELERATION_STRUCTURE_DESC, TInlineAllocator<32>> BuildDescs;
-	BuildDescs.Reserve(1);
-
+	D3D12_BUILD_RAYTRACING_ACCELERATION_STRUCTURE_DESC BuildDesc;
 	{
 		const D3D12_GPU_VIRTUAL_ADDRESS BufferAddress = AccelerationStructureBuffer->ResourceLocation.GetGPUVirtualAddress() + Scene.BufferOffset;
 		D3D12_GPU_VIRTUAL_ADDRESS ScratchAddress = ScratchBuffer->ResourceLocation.GetGPUVirtualAddress() + ScratchBufferOffset;
@@ -4571,7 +4555,6 @@ void BuildAccelerationStructure(FD3D12CommandContext& CommandContext,
 			TEXT("TLAS scratch buffer (plus offset) must be aligned to %lld bytes."),
 			GRHIRayTracingScratchBufferAlignment);
 
-		D3D12_BUILD_RAYTRACING_ACCELERATION_STRUCTURE_DESC& BuildDesc = BuildDescs.AddDefaulted_GetRef();
 		BuildDesc.Inputs = {};
 		BuildDesc.Inputs.Type = D3D12_RAYTRACING_ACCELERATION_STRUCTURE_TYPE_TOP_LEVEL;
 		BuildDesc.Inputs.DescsLayout = D3D12_ELEMENTS_LAYOUT_ARRAY;
@@ -4601,7 +4584,12 @@ void BuildAccelerationStructure(FD3D12CommandContext& CommandContext,
 		}
 	}
 
-	CommandContext.BuildAccelerationStructuresInternal(BuildDescs);
+	// UAV barrier is used here to ensure that all bottom level acceleration structures are built
+	CommandContext.AddUAVBarrier();
+	CommandContext.FlushResourceBarriers();
+
+	// Build the actual acceleration structure
+	CommandContext.BuildAccelerationStructuresInternal(MakeConstArrayView(&BuildDesc, 1));
 
 	// UAV barrier is used here to ensure that the acceleration structure build is complete before any rays are traced
 	// #dxr_todo: these barriers should ideally be inserted by the high level code to allow more overlapped execution
diff --git a/Engine/Source/Runtime/D3D12RHI/Private/D3D12RayTracingDebug.cpp b/Engine/Source/Runtime/D3D12RHI/Private/D3D12RayTracingDebug.cpp
--- a/Engine/Source/Runtime/D3D12RHI/Private/D3D12RayTracingDebug.cpp
+++ b/Engine/Source/Runtime/D3D12RHI/Private/D3D12RayTracingDebug.cpp
@@ -354,7 +354,7 @@ void D3D12RayTracingSceneDebugUpdate(const FD3D12RayTracingScene& Scene, FD3D12B
 {
 	if (GRayTracingSerializeSceneNextFrame)
 	{
-		// #dxr_todo: this could serialize all scenes until the frame counter is increased (if necessary in the future)
+		// #dxr_todo: Move this to RHICmdList.SerializeAccelerationStructure(...)
 
 		DebugSerializeScene(Scene, InstanceBuffer, InstanceBufferOffset, CommandContext);
 		GRayTracingSerializeSceneNextFrame = false;
