[HWRT] Support batched TLAS builds.

Saves around 0.1ms on consoles.

diff --git a/Engine/Source/Runtime/Apple/MetalRHI/Private/MetalRayTracing.cpp b/Engine/Source/Runtime/Apple/MetalRHI/Private/MetalRayTracing.cpp
--- a/Engine/Source/Runtime/Apple/MetalRHI/Private/MetalRayTracing.cpp
+++ b/Engine/Source/Runtime/Apple/MetalRHI/Private/MetalRayTracing.cpp
@@ -538,15 +538,18 @@ void FMetalRayTracingScene::BuildAccelerationStructure(
 }
 
 
-void FMetalRHICommandContext::RHIBuildAccelerationStructure(const FRayTracingSceneBuildParams& SceneBuildParams)
-{
-	FMetalRayTracingScene* const Scene = ResourceCast(SceneBuildParams.Scene);
-	FMetalRHIBuffer* const ScratchBuffer = ResourceCast(SceneBuildParams.ScratchBuffer);
-	FMetalRHIBuffer* const InstanceBuffer = ResourceCast(SceneBuildParams.InstanceBuffer);
-	Scene->BuildAccelerationStructure(
-		*this,
-		ScratchBuffer, SceneBuildParams.ScratchBufferOffset,
-		InstanceBuffer, SceneBuildParams.InstanceBufferOffset);
+void FMetalRHICommandContext::RHIBuildAccelerationStructure(TConstArrayView<FRayTracingSceneBuildParams> Params)
+{
+	for (const FRayTracingSceneBuildParams& SceneBuildParams : Params)
+	{
+		FMetalRayTracingScene* const Scene = ResourceCast(SceneBuildParams.Scene);
+		FMetalRHIBuffer* const ScratchBuffer = ResourceCast(SceneBuildParams.ScratchBuffer);
+		FMetalRHIBuffer* const InstanceBuffer = ResourceCast(SceneBuildParams.InstanceBuffer);
+		Scene->BuildAccelerationStructure(
+			*this,
+			ScratchBuffer, SceneBuildParams.ScratchBufferOffset,
+			InstanceBuffer, SceneBuildParams.InstanceBufferOffset);
+	}
 }
 
 void FMetalRHICommandContext::RHIBuildAccelerationStructures(TConstArrayView<FRayTracingGeometryBuildParams> Params, const FRHIBufferRange& ScratchBufferRange)
diff --git a/Engine/Source/Runtime/Apple/MetalRHI/Public/MetalRHIContext.h b/Engine/Source/Runtime/Apple/MetalRHI/Public/MetalRHIContext.h
--- a/Engine/Source/Runtime/Apple/MetalRHI/Public/MetalRHIContext.h
+++ b/Engine/Source/Runtime/Apple/MetalRHI/Public/MetalRHIContext.h
@@ -174,7 +174,7 @@ public:
 #if METAL_RHI_RAYTRACING
 	virtual void RHIBindAccelerationStructureMemory(FRHIRayTracingScene* Scene, FRHIBuffer* Buffer, uint32 BufferOffset) final override;
 	virtual void RHIBuildAccelerationStructures(TConstArrayView<FRayTracingGeometryBuildParams> Params, const FRHIBufferRange& ScratchBufferRange) final override;
-	virtual void RHIBuildAccelerationStructure(const FRayTracingSceneBuildParams& SceneBuildParams) final override;
+	virtual void RHIBuildAccelerationStructures(TConstArrayView<FRayTracingSceneBuildParams> Params) final override;
 	virtual void RHIClearRayTracingBindings(FRHIRayTracingScene* Scene) final override;
 	virtual void RHIClearShaderBindingTable(FRHIShaderBindingTable* SBT) final override;
 	virtual void RHIRayTraceDispatch(FRHIRayTracingPipelineState* RayTracingPipelineState, FRHIRayTracingShader* RayGenShader,
diff --git a/Engine/Source/Runtime/D3D12RHI/Private/D3D12CommandContext.h b/Engine/Source/Runtime/D3D12RHI/Private/D3D12CommandContext.h
--- a/Engine/Source/Runtime/D3D12RHI/Private/D3D12CommandContext.h
+++ b/Engine/Source/Runtime/D3D12RHI/Private/D3D12CommandContext.h
@@ -737,7 +737,7 @@ public:
 	static void UnregisterAccelerationStructuresInternalMGPU(TConstArrayView<FRayTracingGeometryBuildParams> Params, FRHIGPUMask GPUMask);
 #endif
 	virtual void RHIBuildAccelerationStructures(TConstArrayView<FRayTracingGeometryBuildParams> Params, const FRHIBufferRange& ScratchBufferRange) final override;
-	virtual void RHIBuildAccelerationStructure(const FRayTracingSceneBuildParams& SceneBuildParams) final override;
+	virtual void RHIBuildAccelerationStructures(TConstArrayView<FRayTracingSceneBuildParams> Params) final override;
 	virtual void RHIClearRayTracingBindings(FRHIRayTracingScene* Scene) final override;
 	virtual void RHICommitRayTracingBindings(FRHIRayTracingScene* Scene) final override;
 	virtual void RHIClearShaderBindingTable(FRHIShaderBindingTable* SBT) final override;
@@ -1099,9 +1099,9 @@ public:
 		ContextRedirect(RHIBuildAccelerationStructures(Params, ScratchBufferRange));
 	}
 
-	virtual void RHIBuildAccelerationStructure(const FRayTracingSceneBuildParams& SceneBuildParams) final override
+	virtual void RHIBuildAccelerationStructures(TConstArrayView<FRayTracingSceneBuildParams> Params) final override
 	{
-		ContextRedirect(RHIBuildAccelerationStructure(SceneBuildParams));
+		ContextRedirect(RHIBuildAccelerationStructures(Params));
 	}
 
 	virtual void RHIRayTraceDispatch(FRHIRayTracingPipelineState* RayTracingPipelineState, FRHIRayTracingShader* RayGenShader,
diff --git a/Engine/Source/Runtime/D3D12RHI/Private/D3D12RayTracing.cpp b/Engine/Source/Runtime/D3D12RHI/Private/D3D12RayTracing.cpp
--- a/Engine/Source/Runtime/D3D12RHI/Private/D3D12RayTracing.cpp
+++ b/Engine/Source/Runtime/D3D12RHI/Private/D3D12RayTracing.cpp
@@ -4350,15 +4350,16 @@ void FD3D12RayTracingScene::BindBuffer(FRHIBuffer* InBuffer, uint32 InBufferOffs
 	BufferOffset = InBufferOffset;
 }
 
-void BuildAccelerationStructure(FD3D12CommandContext& CommandContext,
+void PrepareAccelerationStructureBuild(
+	FD3D12CommandContext& CommandContext,
 	FD3D12RayTracingScene& Scene,
 	FD3D12Buffer* ScratchBuffer, uint32 ScratchBufferOffset,
 	FD3D12Buffer* InstanceBuffer, uint32 InstanceBufferOffset,
 	uint32 NumInstances,
-	EAccelerationStructureBuildMode BuildMode)
+	EAccelerationStructureBuildMode BuildMode,
+	D3D12_BUILD_RAYTRACING_ACCELERATION_STRUCTURE_DESC& OutBuildDesc)
 {
-	TRACE_CPUPROFILER_EVENT_SCOPE(BuildAccelerationStructure_TopLevel);
-	SCOPE_CYCLE_COUNTER(STAT_D3D12BuildTLAS);
+	TRACE_CPUPROFILER_EVENT_SCOPE(PrepareAccelerationStructureBuild_TopLevel);
 
 	check(InstanceBuffer != nullptr);
 	checkf(NumInstances <= Scene.Initializer.MaxNumInstances, TEXT("NumInstances must be less or equal to MaxNumInstances"));
@@ -4542,7 +4543,6 @@ void BuildAccelerationStructure(FD3D12CommandContext& CommandContext,
 			InstanceBuffer, InstanceBufferOffset, InstanceBufferStride);
 	}
 
-	D3D12_BUILD_RAYTRACING_ACCELERATION_STRUCTURE_DESC BuildDesc;
 	{
 		const D3D12_GPU_VIRTUAL_ADDRESS BufferAddress = AccelerationStructureBuffer->ResourceLocation.GetGPUVirtualAddress() + Scene.BufferOffset;
 		D3D12_GPU_VIRTUAL_ADDRESS ScratchAddress = ScratchBuffer->ResourceLocation.GetGPUVirtualAddress() + ScratchBufferOffset;
@@ -4555,24 +4555,24 @@ void BuildAccelerationStructure(FD3D12CommandContext& CommandContext,
 			TEXT("TLAS scratch buffer (plus offset) must be aligned to %lld bytes."),
 			GRHIRayTracingScratchBufferAlignment);
 
-		BuildDesc.Inputs = {};
-		BuildDesc.Inputs.Type = D3D12_RAYTRACING_ACCELERATION_STRUCTURE_TYPE_TOP_LEVEL;
-		BuildDesc.Inputs.DescsLayout = D3D12_ELEMENTS_LAYOUT_ARRAY;
-		BuildDesc.Inputs.NumDescs = NumInstances;
-		BuildDesc.Inputs.InstanceDescs = InstanceBuffer->ResourceLocation.GetGPUVirtualAddress() + InstanceBufferOffset;
-		BuildDesc.Inputs.Flags = TranslateRayTracingAccelerationStructureFlags(Scene.Initializer.BuildFlags);
+		OutBuildDesc.Inputs = {};
+		OutBuildDesc.Inputs.Type = D3D12_RAYTRACING_ACCELERATION_STRUCTURE_TYPE_TOP_LEVEL;
+		OutBuildDesc.Inputs.DescsLayout = D3D12_ELEMENTS_LAYOUT_ARRAY;
+		OutBuildDesc.Inputs.NumDescs = NumInstances;
+		OutBuildDesc.Inputs.InstanceDescs = InstanceBuffer->ResourceLocation.GetGPUVirtualAddress() + InstanceBufferOffset;
+		OutBuildDesc.Inputs.Flags = TranslateRayTracingAccelerationStructureFlags(Scene.Initializer.BuildFlags);
 
 		if (bIsUpdate)
 		{
 			checkf(EnumHasAllFlags(Scene.Initializer.BuildFlags, ERayTracingAccelerationStructureFlags::AllowUpdate),
 				TEXT("Acceleration structure must be created with FRayTracingGeometryInitializer::bAllowUpdate=true to perform refit / update."));
 
-			BuildDesc.Inputs.Flags |= D3D12_RAYTRACING_ACCELERATION_STRUCTURE_BUILD_FLAG_PERFORM_UPDATE;
+			OutBuildDesc.Inputs.Flags |= D3D12_RAYTRACING_ACCELERATION_STRUCTURE_BUILD_FLAG_PERFORM_UPDATE;
 		}
 
-		BuildDesc.DestAccelerationStructureData = BufferAddress;
-		BuildDesc.ScratchAccelerationStructureData = ScratchAddress;
-		BuildDesc.SourceAccelerationStructureData = bIsUpdate ? BufferAddress : D3D12_GPU_VIRTUAL_ADDRESS(0);
+		OutBuildDesc.DestAccelerationStructureData = BufferAddress;
+		OutBuildDesc.ScratchAccelerationStructureData = ScratchAddress;
+		OutBuildDesc.SourceAccelerationStructureData = bIsUpdate ? BufferAddress : D3D12_GPU_VIRTUAL_ADDRESS(0);
 
 		if (bIsUpdate)
 		{
@@ -4583,23 +4583,6 @@ void BuildAccelerationStructure(FD3D12CommandContext& CommandContext,
 			INC_DWORD_STAT(STAT_D3D12RayTracingBuiltTLAS);
 		}
 	}
-
-	// UAV barrier is used here to ensure that all bottom level acceleration structures are built
-	CommandContext.AddUAVBarrier();
-	CommandContext.FlushResourceBarriers();
-
-	// Build the actual acceleration structure
-	CommandContext.BuildAccelerationStructuresInternal(MakeConstArrayView(&BuildDesc, 1));
-
-	// UAV barrier is used here to ensure that the acceleration structure build is complete before any rays are traced
-	// #dxr_todo: these barriers should ideally be inserted by the high level code to allow more overlapped execution
-	CommandContext.AddUAVBarrier();
-
-	Scene.bBuilt = true;
-
-#if D3D12_RHI_SUPPORT_RAYTRACING_SCENE_DEBUGGING
-	D3D12RayTracingSceneDebugUpdate(Scene, InstanceBuffer, InstanceBufferOffset, CommandContext);
-#endif // D3D12_RHI_SUPPORT_RAYTRACING_SCENE_DEBUGGING
 }
 
 void FD3D12RayTracingScene::UpdateResidency(FD3D12CommandContext& CommandContext) const
@@ -4826,27 +4809,62 @@ void FD3D12CommandContext::RHIBuildAccelerationStructures(TConstArrayView<FRayTr
 	AddUAVBarrier();
 }
 
-void FD3D12CommandContext::RHIBuildAccelerationStructure(const FRayTracingSceneBuildParams& SceneBuildParams)
+void FD3D12CommandContext::RHIBuildAccelerationStructures(TConstArrayView<FRayTracingSceneBuildParams> Params)
 {
-	FD3D12RayTracingScene* Scene = FD3D12DynamicRHI::ResourceCast(SceneBuildParams.Scene);
-	FD3D12Buffer* ScratchBuffer = RetrieveObject<FD3D12Buffer>(SceneBuildParams.ScratchBuffer);
-	FD3D12Buffer* InstanceBuffer = RetrieveObject<FD3D12Buffer>(SceneBuildParams.InstanceBuffer);
+	TRACE_CPUPROFILER_EVENT_SCOPE(BuildAccelerationStructures_TopLevel);
+	SCOPE_CYCLE_COUNTER(STAT_D3D12BuildTLAS);
 
-	Scene->ReferencedGeometries.Reserve(SceneBuildParams.ReferencedGeometries.Num());
+	TArray<D3D12_BUILD_RAYTRACING_ACCELERATION_STRUCTURE_DESC, TInlineAllocator<8>> BuildDescs;
+	BuildDescs.Reserve(Params.Num());
 
-	for (FRHIRayTracingGeometry* ReferencedGeometry : SceneBuildParams.ReferencedGeometries)
+	for (const FRayTracingSceneBuildParams& SceneBuildParams : Params)
 	{
-		Scene->ReferencedGeometries.Add(ReferencedGeometry);
+		FD3D12RayTracingScene* Scene = FD3D12DynamicRHI::ResourceCast(SceneBuildParams.Scene);
+		FD3D12Buffer* ScratchBuffer = RetrieveObject<FD3D12Buffer>(SceneBuildParams.ScratchBuffer);
+		FD3D12Buffer* InstanceBuffer = RetrieveObject<FD3D12Buffer>(SceneBuildParams.InstanceBuffer);
+
+		Scene->ReferencedGeometries.Reserve(SceneBuildParams.ReferencedGeometries.Num());
+
+		for (FRHIRayTracingGeometry* ReferencedGeometry : SceneBuildParams.ReferencedGeometries)
+		{
+			Scene->ReferencedGeometries.Add(ReferencedGeometry);
+		}
+
+		D3D12_BUILD_RAYTRACING_ACCELERATION_STRUCTURE_DESC& BuildDesc = BuildDescs.AddDefaulted_GetRef();
+
+		PrepareAccelerationStructureBuild(
+			*this,
+			*Scene,
+			ScratchBuffer, SceneBuildParams.ScratchBufferOffset,
+			InstanceBuffer, SceneBuildParams.InstanceBufferOffset,
+			SceneBuildParams.NumInstances,
+			SceneBuildParams.BuildMode,
+			BuildDesc
+		);
 	}
 
-	BuildAccelerationStructure(
-		*this,
-		*Scene,
-		ScratchBuffer, SceneBuildParams.ScratchBufferOffset,
-		InstanceBuffer, SceneBuildParams.InstanceBufferOffset,
-		SceneBuildParams.NumInstances,
-		SceneBuildParams.BuildMode
-	);
+	// UAV barrier is used here to ensure that all bottom level acceleration structures are built
+	AddUAVBarrier();
+	FlushResourceBarriers();
+
+	BuildAccelerationStructuresInternal(BuildDescs);
+
+	// UAV barrier is used here to ensure that the acceleration structure build is complete before any rays are traced
+	// #dxr_todo: these barriers should ideally be inserted by the high level code to allow more overlapped execution
+	AddUAVBarrier();
+
+	for (const FRayTracingSceneBuildParams& SceneBuildParams : Params)
+	{
+		FD3D12RayTracingScene* Scene = FD3D12DynamicRHI::ResourceCast(SceneBuildParams.Scene);
+		FD3D12Buffer* ScratchBuffer = RetrieveObject<FD3D12Buffer>(SceneBuildParams.ScratchBuffer);
+		FD3D12Buffer* InstanceBuffer = RetrieveObject<FD3D12Buffer>(SceneBuildParams.InstanceBuffer);
+
+		Scene->bBuilt = true;
+
+#if D3D12_RHI_SUPPORT_RAYTRACING_SCENE_DEBUGGING
+		D3D12RayTracingSceneDebugUpdate(*Scene, InstanceBuffer, SceneBuildParams.InstanceBufferOffset, *this);
+#endif // D3D12_RHI_SUPPORT_RAYTRACING_SCENE_DEBUGGING
+	}
 }
 
 void FD3D12CommandContext::RHIBindAccelerationStructureMemory(FRHIRayTracingScene* InScene, FRHIBuffer* InBuffer, uint32 InBufferOffset)
diff --git a/Engine/Source/Runtime/RHI/Public/RHICommandList.h b/Engine/Source/Runtime/RHI/Public/RHICommandList.h
--- a/Engine/Source/Runtime/RHI/Public/RHICommandList.h
+++ b/Engine/Source/Runtime/RHI/Public/RHICommandList.h
@@ -2586,12 +2586,12 @@ FRHICOMMAND_UNNAMED(FRHICommandBindAccelerationStructureMemory)
 	RHI_EXECUTE_API void Execute(FRHICommandListBase& CmdList);
 };
 
-FRHICOMMAND_UNNAMED(FRHICommandBuildAccelerationStructure)
+FRHICOMMAND_UNNAMED(FRHICommandBuildSceneAccelerationStructures)
 {
-	FRayTracingSceneBuildParams SceneBuildParams;
+	TConstArrayView<FRayTracingSceneBuildParams> Params;
 
-	explicit FRHICommandBuildAccelerationStructure(FRayTracingSceneBuildParams InSceneBuildParams)
-		: SceneBuildParams(MoveTemp(InSceneBuildParams))
+	explicit FRHICommandBuildSceneAccelerationStructures(TConstArrayView<FRayTracingSceneBuildParams> InParams)
+		: Params(InParams)
 	{}
 
 	RHI_EXECUTE_API void Execute(FRHICommandListBase& CmdList);
@@ -3810,20 +3810,38 @@ public:
 	}
 
 	FORCEINLINE_DEBUGGABLE void BuildAccelerationStructure(const FRayTracingSceneBuildParams& SceneBuildParams)
+	{
+		BuildAccelerationStructures(MakeConstArrayView(&SceneBuildParams, 1));
+	}
+
+	FORCEINLINE_DEBUGGABLE void BuildAccelerationStructures(TConstArrayView<FRayTracingSceneBuildParams> Params)
 	{
 		if (Bypass())
 		{
-			GetComputeContext().RHIBuildAccelerationStructure(SceneBuildParams);
+			GetComputeContext().RHIBuildAccelerationStructures(Params);
 		}
 		else
 		{
-			FRayTracingSceneBuildParams InlineParams = SceneBuildParams;
-			InlineParams.ReferencedGeometries = AllocArray(SceneBuildParams.ReferencedGeometries);
-			PRAGMA_DISABLE_DEPRECATION_WARNINGS
-			InlineParams.PerInstanceGeometries = AllocArray(SceneBuildParams.PerInstanceGeometries);
-			PRAGMA_ENABLE_DEPRECATION_WARNINGS
+			// Copy the params themselves as well their ReferencedGeometries, if there are any.
+			// AllocArray() can't be used here directly, as we have to modify the params after copy.
+			size_t DataSize = sizeof(FRayTracingSceneBuildParams) * Params.Num();
+			FRayTracingSceneBuildParams* InlineParams = (FRayTracingSceneBuildParams*)Alloc(DataSize, alignof(FRayTracingSceneBuildParams));
+			FMemory::Memcpy(InlineParams, Params.GetData(), DataSize);
+			for (int32 Index = 0; Index < Params.Num(); ++Index)
+			{
+				if (Params[Index].ReferencedGeometries.Num())
+				{
+					InlineParams[Index].ReferencedGeometries = AllocArray(Params[Index].ReferencedGeometries);
+				}
+				PRAGMA_DISABLE_DEPRECATION_WARNINGS
+				if (Params[Index].PerInstanceGeometries.Num())
+				{
+					InlineParams[Index].PerInstanceGeometries = AllocArray(Params[Index].PerInstanceGeometries);
+				}
+				PRAGMA_ENABLE_DEPRECATION_WARNINGS
+			}
 
-			ALLOC_COMMAND(FRHICommandBuildAccelerationStructure)(MoveTemp(InlineParams));
+			ALLOC_COMMAND(FRHICommandBuildSceneAccelerationStructures)(MakeConstArrayView(InlineParams, Params.Num()));
 
 			// This RHI command modifies members of the FRHIRayTracingScene inside platform RHI implementations.
 			// It therefore needs the RHI lock fence to prevent races on those members.
diff --git a/Engine/Source/Runtime/RHI/Public/RHICommandListCommandExecutes.inl b/Engine/Source/Runtime/RHI/Public/RHICommandListCommandExecutes.inl
--- a/Engine/Source/Runtime/RHI/Public/RHICommandListCommandExecutes.inl
+++ b/Engine/Source/Runtime/RHI/Public/RHICommandListCommandExecutes.inl
@@ -502,10 +502,10 @@ void FRHICommandBindAccelerationStructureMemory::Execute(FRHICommandListBase& Cm
 	INTERNAL_DECORATOR_COMPUTE(RHIBindAccelerationStructureMemory)(Scene, Buffer, BufferOffset);
 }
 
-void FRHICommandBuildAccelerationStructure::Execute(FRHICommandListBase& CmdList)
+void FRHICommandBuildSceneAccelerationStructures::Execute(FRHICommandListBase& CmdList)
 {
 	RHISTAT(BuildAccelerationStructure);
-	INTERNAL_DECORATOR_COMPUTE(RHIBuildAccelerationStructure)(SceneBuildParams);
+	INTERNAL_DECORATOR_COMPUTE(RHIBuildAccelerationStructures)(Params);
 }
 
 void FRHICommandCommitRayTracingBindings::Execute(FRHICommandListBase& CmdList)
diff --git a/Engine/Source/Runtime/RHI/Public/RHIContext.h b/Engine/Source/Runtime/RHI/Public/RHIContext.h
--- a/Engine/Source/Runtime/RHI/Public/RHIContext.h
+++ b/Engine/Source/Runtime/RHI/Public/RHIContext.h
@@ -498,7 +498,7 @@ public:
 		checkNoEntry();
 	}
 
-	virtual void RHIBuildAccelerationStructure(const FRayTracingSceneBuildParams& SceneBuildParams)
+	virtual void RHIBuildAccelerationStructures(TConstArrayView<FRayTracingSceneBuildParams> Params)
 	{
 		checkNoEntry();
 	}
diff --git a/Engine/Source/Runtime/RHI/Public/RHIValidationContext.h b/Engine/Source/Runtime/RHI/Public/RHIValidationContext.h
--- a/Engine/Source/Runtime/RHI/Public/RHIValidationContext.h
+++ b/Engine/Source/Runtime/RHI/Public/RHIValidationContext.h
@@ -388,7 +388,7 @@ public:
 
 	virtual void RHIBuildAccelerationStructures(TConstArrayView<FRayTracingGeometryBuildParams> Params, const FRHIBufferRange& ScratchBufferRange) override final
 	{
-		// #yuriy_todo: explicit transitions and state validation for BLAS
+		// #hwrt_todo: explicit transitions and state validation for BLAS
 		for (const FRayTracingGeometryBuildParams& P : Params)
 		{
 			const FRayTracingGeometryInitializer& Initializer = P.Geometry->GetInitializer();
@@ -408,25 +408,28 @@ public:
 		RHIContext->RHIBuildAccelerationStructures(Params, ScratchBufferRange);
 	}
 
-	virtual void RHIBuildAccelerationStructure(const FRayTracingSceneBuildParams& SceneBuildParams) override final
+	virtual void RHIBuildAccelerationStructures(TConstArrayView<FRayTracingSceneBuildParams> Params) override final
 	{
-		// #yuriy_todo: validate all referenced BLAS states
-		if (SceneBuildParams.Scene)
+		// #hwrt_todo: validate all referenced BLAS states
+		for (const FRayTracingSceneBuildParams& P : Params)
 		{
-			Tracker->Assert(SceneBuildParams.Scene->GetWholeResourceIdentity(), ERHIAccess::BVHWrite);
-		}
+			if (P.Scene)
+			{
+				Tracker->Assert(P.Scene->GetWholeResourceIdentity(), ERHIAccess::BVHWrite);
+			}
 
-		if (SceneBuildParams.InstanceBuffer)
-		{
-			Tracker->Assert(SceneBuildParams.InstanceBuffer->GetWholeResourceIdentity(), ERHIAccess::SRVCompute);
-		}
+			if (P.InstanceBuffer)
+			{
+				Tracker->Assert(P.InstanceBuffer->GetWholeResourceIdentity(), ERHIAccess::SRVCompute);
+			}
 
-		if (SceneBuildParams.ScratchBuffer)
-		{
-			Tracker->Assert(SceneBuildParams.ScratchBuffer->GetWholeResourceIdentity(), ERHIAccess::UAVCompute);
+			if (P.ScratchBuffer)
+			{
+				Tracker->Assert(P.ScratchBuffer->GetWholeResourceIdentity(), ERHIAccess::UAVCompute);
+			}
 		}
 		
-		RHIContext->RHIBuildAccelerationStructure(SceneBuildParams);
+		RHIContext->RHIBuildAccelerationStructures(Params);
 	}
 
 	virtual void RHIBindAccelerationStructureMemory(FRHIRayTracingScene* Scene, FRHIBuffer* Buffer, uint32 BufferOffset) override final
@@ -1261,24 +1264,27 @@ public:
 		RHIContext->RHIBuildAccelerationStructures(Params, ScratchBufferRange);
 	}
 
-	virtual void RHIBuildAccelerationStructure(const FRayTracingSceneBuildParams& SceneBuildParams) override final
+	virtual void RHIBuildAccelerationStructures(TConstArrayView<FRayTracingSceneBuildParams> Params) override final
 	{
-		if (SceneBuildParams.Scene)
+		for (const FRayTracingSceneBuildParams& P : Params)
 		{
-			Tracker->Assert(SceneBuildParams.Scene->GetWholeResourceIdentity(), ERHIAccess::BVHWrite);
-		}
+			if (P.Scene)
+			{
+				Tracker->Assert(P.Scene->GetWholeResourceIdentity(), ERHIAccess::BVHWrite);
+			}
 
-		if (SceneBuildParams.InstanceBuffer)
-		{
-			Tracker->Assert(SceneBuildParams.InstanceBuffer->GetWholeResourceIdentity(), ERHIAccess::SRVCompute);
-		}
+			if (P.InstanceBuffer)
+			{
+				Tracker->Assert(P.InstanceBuffer->GetWholeResourceIdentity(), ERHIAccess::SRVCompute);
+			}
 
-		if (SceneBuildParams.ScratchBuffer)
-		{
-			Tracker->Assert(SceneBuildParams.ScratchBuffer->GetWholeResourceIdentity(), ERHIAccess::UAVCompute);
+			if (P.ScratchBuffer)
+			{
+				Tracker->Assert(P.ScratchBuffer->GetWholeResourceIdentity(), ERHIAccess::UAVCompute);
+			}
 		}
 
-		RHIContext->RHIBuildAccelerationStructure(SceneBuildParams);
+		RHIContext->RHIBuildAccelerationStructures(Params);
 	}
 
 	virtual void RHIBindAccelerationStructureMemory(FRHIRayTracingScene* Scene, FRHIBuffer* Buffer, uint32 BufferOffset) override final
diff --git a/Engine/Source/Runtime/Renderer/Private/RayTracing/RayTracingScene.cpp b/Engine/Source/Runtime/Renderer/Private/RayTracing/RayTracingScene.cpp
--- a/Engine/Source/Runtime/Renderer/Private/RayTracing/RayTracingScene.cpp
+++ b/Engine/Source/Runtime/Renderer/Private/RayTracing/RayTracingScene.cpp
@@ -32,6 +32,13 @@ static TAutoConsoleVariable<bool> CVarRayTracingSceneUseTracingFeedback(
 	ECVF_RenderThreadSafe
 );
 
+static TAutoConsoleVariable<bool> CVarRayTracingSceneBatchedBuild(
+	TEXT("r.RayTracing.Scene.BatchedBuild"),
+	true,
+	TEXT("Whether to batch TLAS builds. Should be kept enabled since batched builds reduce barriers on GPU."),
+	ECVF_RenderThreadSafe
+);
+
 static TAutoConsoleVariable<bool> CVarRayTracingSceneCompactInstances(
 	TEXT("r.RayTracing.Scene.CompactInstances"),
 	false,
@@ -342,10 +349,7 @@ void FRayTracingScene::Create(FRDGBuilder& GraphBuilder, const FViewInfo& View,
 }
 
 BEGIN_SHADER_PARAMETER_STRUCT(FRayTracingSceneBuildPassParams, )
-	RDG_BUFFER_ACCESS(ScratchBuffer, ERHIAccess::UAVCompute)
-	RDG_BUFFER_ACCESS(InstanceBuffer, ERHIAccess::SRVCompute)
-	RDG_BUFFER_ACCESS(TLASBuffer, ERHIAccess::BVHWrite)
-
+	RDG_BUFFER_ACCESS_ARRAY(TLASBuildBuffers)
 	RDG_BUFFER_ACCESS(DynamicGeometryScratchBuffer, ERHIAccess::UAVCompute)
 END_SHADER_PARAMETER_STRUCT()
 
@@ -356,34 +360,63 @@ END_SHADER_PARAMETER_STRUCT()
 void FRayTracingScene::Build(FRDGBuilder& GraphBuilder, ERDGPassFlags ComputePassFlags, FRDGBufferRef DynamicGeometryScratchBuffer)
 {
 	const uint8 NumLayers = uint8(ERayTracingSceneLayer::NUM);
-	
+
+	FRayTracingSceneBuildPassParams* PassParams = GraphBuilder.AllocParameters<FRayTracingSceneBuildPassParams>();
+	PassParams->DynamicGeometryScratchBuffer = DynamicGeometryScratchBuffer; // TODO: Is this necessary?
+
 	for (uint32 LayerIndex = 0; LayerIndex < NumLayers; ++LayerIndex)
 	{
 		FLayer& Layer = Layers[LayerIndex];
 
-		FRayTracingSceneBuildPassParams* PassParams = GraphBuilder.AllocParameters<FRayTracingSceneBuildPassParams>();
-		PassParams->ScratchBuffer = Layer.BuildScratchBuffer;
-		PassParams->InstanceBuffer = Layer.InstanceBuffer;
-		PassParams->TLASBuffer = Layer.RayTracingSceneBufferRDG;
-		PassParams->DynamicGeometryScratchBuffer = DynamicGeometryScratchBuffer; // TODO: Is this necessary?
+		PassParams->TLASBuildBuffers.Emplace(Layer.BuildScratchBuffer, ERHIAccess::UAVCompute);
+		PassParams->TLASBuildBuffers.Emplace(Layer.InstanceBuffer, ERHIAccess::SRVCompute);
+		PassParams->TLASBuildBuffers.Emplace(Layer.RayTracingSceneBufferRDG, ERHIAccess::BVHWrite);
+	}
+
+	GraphBuilder.AddPass(RDG_EVENT_NAME("RayTracingBuildScene"), PassParams, ComputePassFlags,
+		[PassParams, this](FRHICommandList& RHICmdList)
+		{
+			const bool bUseBatchedBuild = CVarRayTracingSceneBatchedBuild.GetValueOnRenderThread();
+
+			TArray<FRayTracingSceneBuildParams> BatchedBuildParams;
+			BatchedBuildParams.Reserve(NumLayers);
 
-		GraphBuilder.AddPass(RDG_EVENT_NAME("RayTracingBuildScene"), PassParams, ComputePassFlags,
-			[PassParams, &Layer](FRHICommandList& RHICmdList)
+			for (uint32 LayerIndex = 0; LayerIndex < NumLayers; ++LayerIndex)
 			{
+				FLayer& Layer = Layers[LayerIndex];
+
 				FRayTracingSceneBuildParams BuildParams;
 				BuildParams.Scene = Layer.RayTracingSceneRHI;
-				BuildParams.ScratchBuffer = PassParams->ScratchBuffer->GetRHI();
+				BuildParams.ScratchBuffer = Layer.BuildScratchBuffer->GetRHI();
 				BuildParams.ScratchBufferOffset = 0;
-				BuildParams.InstanceBuffer = PassParams->InstanceBuffer->GetRHI();
+				BuildParams.InstanceBuffer = Layer.InstanceBuffer->GetRHI();
 				BuildParams.InstanceBufferOffset = 0;
 				BuildParams.NumInstances = Layer.MaxNumInstances;
 				BuildParams.ReferencedGeometries = Layer.InstanceBufferBuilder.GetReferencedGeometries();
 
-				RHICmdList.BindAccelerationStructureMemory(Layer.RayTracingSceneRHI, PassParams->TLASBuffer->GetRHI(), 0);
-				RHICmdList.BuildAccelerationStructure(BuildParams);
-			});
+				RHICmdList.BindAccelerationStructureMemory(Layer.RayTracingSceneRHI, Layer.RayTracingSceneBufferRDG->GetRHI(), 0);
+
+				if (bUseBatchedBuild)
+				{
+					BatchedBuildParams.Add(BuildParams);
+				}
+				else
+				{
+					RHICmdList.BuildAccelerationStructure(BuildParams);
+				}
+			}
+
+			if (bUseBatchedBuild)
+			{
+				RHICmdList.BuildAccelerationStructures(BatchedBuildParams);
+			}
+		});
 
 #if !UE_BUILD_SHIPPING
+	for (uint32 LayerIndex = 0; LayerIndex < NumLayers; ++LayerIndex)
+	{
+		FLayer& Layer = Layers[LayerIndex];
+
 		if (GRayTracingSerializeSceneNextFrame && GRHIGlobals.RayTracing.SupportsSerializeAccelerationStructure)
 		{
 			FRayTracingSceneSerializePassParams* SerializePassParams = GraphBuilder.AllocParameters<FRayTracingSceneSerializePassParams>();
@@ -399,12 +432,10 @@ void FRayTracingScene::Build(FRDGBuilder& GraphBuilder, ERDGPassFlags ComputePas
 					RHICmdList.SerializeAccelerationStructure(Layer.RayTracingSceneRHI, *OutputFilename);
 				});
 		}
-#endif
 	}
 
-#if !UE_BUILD_SHIPPING
 	GRayTracingSerializeSceneNextFrame = false;
-#endif
+#endif // !UE_BUILD_SHIPPING
 }
 
 struct FRayTracingProcessFeedbackCS : public FGlobalShader
diff --git a/Engine/Source/Runtime/VulkanRHI/Private/VulkanContext.h b/Engine/Source/Runtime/VulkanRHI/Private/VulkanContext.h
--- a/Engine/Source/Runtime/VulkanRHI/Private/VulkanContext.h
+++ b/Engine/Source/Runtime/VulkanRHI/Private/VulkanContext.h
@@ -333,7 +333,7 @@ public:
 	virtual void RHICommitShaderBindingTable(FRHIShaderBindingTable* SBT, FRHIBuffer* InlineBindingDataBuffer) final override;
 	virtual void RHIBindAccelerationStructureMemory(FRHIRayTracingScene* Scene, FRHIBuffer* Buffer, uint32 BufferOffset) final override;
 	virtual void RHIBuildAccelerationStructures(TConstArrayView<FRayTracingGeometryBuildParams> Params, const FRHIBufferRange& ScratchBufferRange) final override;
-	virtual void RHIBuildAccelerationStructure(const FRayTracingSceneBuildParams& SceneBuildParams) final override;
+	virtual void RHIBuildAccelerationStructures(TConstArrayView<FRayTracingSceneBuildParams> Params) final override;
 
 	virtual void RHIRayTraceDispatch(FRHIRayTracingPipelineState* RayTracingPipelineState, FRHIRayTracingShader* RayGenShader,
 		FRHIShaderBindingTable* SBT, const FRayTracingShaderBindings& GlobalResourceBindings,
diff --git a/Engine/Source/Runtime/VulkanRHI/Private/VulkanRayTracing.cpp b/Engine/Source/Runtime/VulkanRHI/Private/VulkanRayTracing.cpp
--- a/Engine/Source/Runtime/VulkanRHI/Private/VulkanRayTracing.cpp
+++ b/Engine/Source/Runtime/VulkanRHI/Private/VulkanRayTracing.cpp
@@ -1333,26 +1333,29 @@ void FVulkanCommandListContext::RHIBuildAccelerationStructures(TConstArrayView<F
 	}
 }
 
-void FVulkanCommandListContext::RHIBuildAccelerationStructure(const FRayTracingSceneBuildParams& SceneBuildParams)
+void FVulkanCommandListContext::RHIBuildAccelerationStructures(TConstArrayView<FRayTracingSceneBuildParams> Params)
 {
-	FVulkanRayTracingScene* const Scene = ResourceCast(SceneBuildParams.Scene);
-	FVulkanBuffer* const ScratchBuffer = ResourceCast(SceneBuildParams.ScratchBuffer);
-	FVulkanBuffer* const InstanceBuffer = ResourceCast(SceneBuildParams.InstanceBuffer);
+	for (const FRayTracingSceneBuildParams& SceneBuildParams : Params)
+	{
+		FVulkanRayTracingScene* const Scene = ResourceCast(SceneBuildParams.Scene);
+		FVulkanBuffer* const ScratchBuffer = ResourceCast(SceneBuildParams.ScratchBuffer);
+		FVulkanBuffer* const InstanceBuffer = ResourceCast(SceneBuildParams.InstanceBuffer);
 
-	Scene->ReferencedGeometries.Reserve(SceneBuildParams.ReferencedGeometries.Num());
+		Scene->ReferencedGeometries.Reserve(SceneBuildParams.ReferencedGeometries.Num());
 
-	for (FRHIRayTracingGeometry* ReferencedGeometry : SceneBuildParams.ReferencedGeometries)
-	{
-		Scene->ReferencedGeometries.Add(ReferencedGeometry);
-	}
+		for (FRHIRayTracingGeometry* ReferencedGeometry : SceneBuildParams.ReferencedGeometries)
+		{
+			Scene->ReferencedGeometries.Add(ReferencedGeometry);
+		}
 
-	BuildAccelerationStructure(
-		*this,
-		*Scene,
-		ScratchBuffer, SceneBuildParams.ScratchBufferOffset, 
-		InstanceBuffer, SceneBuildParams.InstanceBufferOffset,
-		SceneBuildParams.NumInstances,
-		SceneBuildParams.BuildMode);
+		BuildAccelerationStructure(
+			*this,
+			*Scene,
+			ScratchBuffer, SceneBuildParams.ScratchBufferOffset,
+			InstanceBuffer, SceneBuildParams.InstanceBufferOffset,
+			SceneBuildParams.NumInstances,
+			SceneBuildParams.BuildMode);
+	}
 }
 
 template<typename ShaderType>
