Fix for d3d12 context closing when only transitions were added to the context.
The bindless descriptor heap management close can come after the common command list closing which prevents issues when opening the context during closing the context.

diff --git a/Engine/Source/Runtime/D3D12RHI/Private/D3D12CommandContext.cpp b/Engine/Source/Runtime/D3D12RHI/Private/D3D12CommandContext.cpp
--- a/Engine/Source/Runtime/D3D12RHI/Private/D3D12CommandContext.cpp
+++ b/Engine/Source/Runtime/D3D12RHI/Private/D3D12CommandContext.cpp
@@ -444,6 +444,12 @@ void FD3D12CommandContext::CloseCommandList()
 {
 	StateCache.GetDescriptorCache()->CloseCommandList();
 	FD3D12ContextCommon::CloseCommandList();
+
+#if PLATFORM_SUPPORTS_BINDLESS_RENDERING
+	// Always call the Bindless Manager CloseCommandList, it will determine when it needs to do anything.
+	GetParentDevice()->GetBindlessDescriptorManager().CloseCommandList(*this);
+#endif
+
 	// Mark state as dirty now, because ApplyState may be called before OpenCommandList(), and it needs to know that the state has
 	// become invalid, so it can set it up again (which opens a new command list if necessary).
 	StateCache.DirtyStateForNewCommandList();
diff --git a/Engine/Source/Runtime/D3D12RHI/Private/D3D12DescriptorCache.cpp b/Engine/Source/Runtime/D3D12RHI/Private/D3D12DescriptorCache.cpp
--- a/Engine/Source/Runtime/D3D12RHI/Private/D3D12DescriptorCache.cpp
+++ b/Engine/Source/Runtime/D3D12RHI/Private/D3D12DescriptorCache.cpp
@@ -183,9 +183,6 @@ void FD3D12DescriptorCache::CloseCommandList()
 	}
 
 #if PLATFORM_SUPPORTS_BINDLESS_RENDERING
-	// Always call the Bindless Manager CloseCommandList, it will determine when it needs to do anything.
-	GetParentDevice()->GetBindlessDescriptorManager().CloseCommandList(Context);
-
 	if (!IsUsingBindlessSamplers())
 #endif
 	{
