Fix false positive asserts in the D3D12 CopyTextureRegionChecked() function.
We need to take into account the special case when copying between UINT and block compressed textures.
The change adjusts the destination box (used only for validation) whenever it has been taken directly from the source box.
#rnx

diff --git a/Engine/Source/Runtime/D3D12RHI/Private/D3D12Texture.cpp b/Engine/Source/Runtime/D3D12RHI/Private/D3D12Texture.cpp
--- a/Engine/Source/Runtime/D3D12RHI/Private/D3D12Texture.cpp
+++ b/Engine/Source/Runtime/D3D12RHI/Private/D3D12Texture.cpp
@@ -1175,6 +1175,18 @@ static D3D12_BOX GetSubresourceCopyBox(const D3D12_RESOURCE_DESC& Desc, const D3
 	};
 }
 
+static D3D12_BOX GetBoxAdjustedForPixelFormat(const D3D12_BOX& SourceBox, const FPixelFormatInfo& SourcePixelFormatInfo, const FPixelFormatInfo& DestPixelFormatInfo)
+{
+	return {
+		.left = SourceBox.left * DestPixelFormatInfo.BlockSizeX / SourcePixelFormatInfo.BlockSizeX,
+		.top = SourceBox.top * DestPixelFormatInfo.BlockSizeY / SourcePixelFormatInfo.BlockSizeY,
+		.front = SourceBox.front,
+		.right = SourceBox.right * DestPixelFormatInfo.BlockSizeX / SourcePixelFormatInfo.BlockSizeX,
+		.bottom = SourceBox.bottom * DestPixelFormatInfo.BlockSizeY / SourcePixelFormatInfo.BlockSizeY,
+		.back = SourceBox.back
+	};
+}
+
 static uint64 GetPlacedBufferRequiredSize(FD3D12CommandContext& Context, const D3D12_SUBRESOURCE_FOOTPRINT& Footprint, const FPixelFormatInfo& Format, const FName& DebugName)
 {
 	RHI_BREADCRUMB_CHECK_SHIPPINGF(Context, (Footprint.Width % Format.BlockSizeX) == 0, TEXT("Width not aligned to block size for: '%s'"), *DebugName.ToString());
@@ -1240,11 +1252,10 @@ FORCENOINLINE void FD3D12CommandContext::CopyTextureRegionChecked(
 		case D3D12_TEXTURE_COPY_TYPE_SUBRESOURCE_INDEX:
 		{
 			const FSubresourceIndices Indices = GetSubresourceIndices(DestDesc, DestCopyLocation->SubresourceIndex);
-
-			const D3D12_BOX CopyBox = SourceBox ? *SourceBox : GetSubresourceCopyBox(DestDesc, SourceCopyLocation, Indices);
+			const D3D12_BOX CopyBox = SourceBox ? GetBoxAdjustedForPixelFormat(*SourceBox, SourcePixelFormatInfo, DestPixelFormatInfo) : GetSubresourceCopyBox(DestDesc, SourceCopyLocation, Indices);
 				
-			RHI_BREADCRUMB_CHECK_SHIPPINGF(*this, (CopyBox.left % DestPixelFormatInfo.BlockSizeX) == 0 && (CopyBox.right % DestPixelFormatInfo.BlockSizeX) == 0, TEXT("Width not aligned to block size for: '%s'"), *DebugName.ToString());
-			RHI_BREADCRUMB_CHECK_SHIPPINGF(*this, (CopyBox.top % DestPixelFormatInfo.BlockSizeY) == 0 && (CopyBox.bottom % DestPixelFormatInfo.BlockSizeY) == 0, TEXT("Height not aligned to block size for: '%s'"), *DebugName.ToString());
+			RHI_BREADCRUMB_CHECK_SHIPPINGF(*this, DestX % DestPixelFormatInfo.BlockSizeX == 0 && (CopyBox.right - CopyBox.left) % DestPixelFormatInfo.BlockSizeX == 0, TEXT("Width not aligned to block size for: '%s'"), *DebugName.ToString());
+			RHI_BREADCRUMB_CHECK_SHIPPINGF(*this, DestY % DestPixelFormatInfo.BlockSizeY == 0 && (CopyBox.bottom - CopyBox.top) % DestPixelFormatInfo.BlockSizeY == 0, TEXT("Height not aligned to block size for: '%s'"), *DebugName.ToString());
 				
 			const uint32 SubresourceWidth  = AlignArbitrary<uint32>(FMath::Max<uint32>(1u, DestDesc.Width >> Indices.MipIndex), DestPixelFormatInfo.BlockSizeX);
 			const uint32 SubresourceHeight = AlignArbitrary<uint32>(FMath::Max<uint32>(1u, DestDesc.Height >> Indices.MipIndex), DestPixelFormatInfo.BlockSizeY);
