Fix crash when trying to transition a BVH.

GetResourceAndTexture() returns null for BVHs so ProcessResource doesn't need to do anything in that case.

#rnx

diff --git a/Engine/Source/Runtime/D3D12RHI/Private/D3D12Commands.cpp b/Engine/Source/Runtime/D3D12RHI/Private/D3D12Commands.cpp
--- a/Engine/Source/Runtime/D3D12RHI/Private/D3D12Commands.cpp
+++ b/Engine/Source/Runtime/D3D12RHI/Private/D3D12Commands.cpp
@@ -292,7 +292,10 @@ void ProcessResource(FD3D12CommandContext& Context, const FRHITransitionInfo& In
 #endif // #if PLATFORM_REQUIRES_TYPELESS_RESOURCE_DISCARD_WORKAROUND
 	}
 
-	Function(Info, Resource, Texture, DiscardTextureOut);
+	if (Resource)
+	{
+		Function(Info, Resource, Texture, DiscardTextureOut);
+	}
 }
 
 // Pipe changes which are not ending with graphics or targeting all pipelines are handle during begin
