Cleaning up RHI per-platform "set resource name" code.
- Textures and Buffers now have their names provided at creation time, so make sure those are set on each platform API.
- A lot of per-platform name setting was done ad hoc, those are now set using helper functions.
- Adding default RHIBindDebugLabelName implementations.
- Adding more enforcement of only setting names when RHI_USE_RESOURCE_DEBUG_NAME is enabled.

diff --git a/Engine/Source/Runtime/Apple/MetalRHI/Private/MetalTexture.cpp b/Engine/Source/Runtime/Apple/MetalRHI/Private/MetalTexture.cpp
--- a/Engine/Source/Runtime/Apple/MetalRHI/Private/MetalTexture.cpp
+++ b/Engine/Source/Runtime/Apple/MetalRHI/Private/MetalTexture.cpp
@@ -605,6 +605,10 @@ FMetalSurface::FMetalSurface(FMetalDevice& MetalDevice, FMetalTextureCreateDesc
 		}
 	}
 
+#if RHI_USE_RESOURCE_DEBUG_NAME
+    NS::String* DebugName = FStringToNSString(CreateDesc.DebugName);
+#endif
+
 	if (!Texture)
 	{
 		// Non VR/media texture case (i.e. a regular texture
@@ -691,6 +695,10 @@ FMetalSurface::FMetalSurface(FMetalDevice& MetalDevice, FMetalTextureCreateDesc
 		}
 		
 		METAL_FATAL_ASSERT(Texture, TEXT("Failed to create texture, desc %s"), *NSStringToFString(CreateDesc.Desc->description()));
+
+#if RHI_USE_RESOURCE_DEBUG_NAME
+		Texture->setLabel(DebugName);
+#endif
 	}
 
 	// calculate size of the texture
@@ -768,6 +776,10 @@ FMetalSurface::FMetalSurface(FMetalDevice& MetalDevice, FMetalTextureCreateDesc
 		{
             UE_LOG(LogMetal, Fatal, TEXT("Failed to create MSAA texture"));
 		}
+
+#if RHI_USE_RESOURCE_DEBUG_NAME
+		MSAATexture->setLabel(DebugName);
+#endif
 	}
 	
 	// create a stencil buffer if needed
@@ -1917,18 +1929,23 @@ void FMetalDynamicRHI::RHIUpdateTexture3D(FRHICommandListBase& RHICmdList, FRHIT
 
 void FMetalDynamicRHI::RHIBindDebugLabelName(FRHICommandListBase& RHICmdList, FRHITexture* TextureRHI, const TCHAR* Name)
 {
+#if RHI_USE_RESOURCE_DEBUG_NAME
 	MTL_SCOPED_AUTORELEASE_POOL;
-    FMetalSurface* Surf = GetMetalSurfaceFromRHITexture(TextureRHI);   
-	
-    NS::String* LabelString = FStringToNSString(Name);
-    if(Surf->Texture)
-    {
-        Surf->Texture->setLabel(LabelString);
-    }
-    if(Surf->MSAATexture)
-    {
-        Surf->MSAATexture->setLabel(LabelString);
-    }
+
+	TextureRHI->SetName(Name);
+
+	FMetalSurface* Surface = GetMetalSurfaceFromRHITexture(TextureRHI);
+
+	NS::String* LabelString = FStringToNSString(Name);
+	if (Surface->Texture)
+	{
+		Surface->Texture->setLabel(LabelString);
+	}
+	if (Surface->MSAATexture)
+	{
+		Surface->MSAATexture->setLabel(LabelString);
+	}
+#endif
 }
 
 inline bool MetalRHICopyTexutre_IsTextureFormatCompatible(EPixelFormat SrcFmt, EPixelFormat DstFmt)
diff --git a/Engine/Source/Runtime/D3D12RHI/Private/D3D12Buffer.cpp b/Engine/Source/Runtime/D3D12RHI/Private/D3D12Buffer.cpp
--- a/Engine/Source/Runtime/D3D12RHI/Private/D3D12Buffer.cpp
+++ b/Engine/Source/Runtime/D3D12RHI/Private/D3D12Buffer.cpp
@@ -823,14 +823,15 @@ void FD3D12DynamicRHI::RHIBindDebugLabelName(FRHICommandListBase& RHICmdList, FR
 		return;
 	}
 
-#if NAME_OBJECTS
+#if RHI_USE_RESOURCE_DEBUG_NAME
+	// Also set on RHI object
+	BufferRHI->SetName(Name);
+
 	FD3D12Buffer* Buffer = FD3D12DynamicRHI::ResourceCast(BufferRHI);
 
 	// only rename the underlying d3d12 resource if it's not sub allocated (requires resource state tracking or stand alone allocated)
-	if (Buffer->GetResource() != nullptr && (Buffer->GetResource()->RequiresResourceStateTracking() || Buffer->ResourceLocation.GetType() == FD3D12ResourceLocation::ResourceLocationType::eStandAlone))
+	if (ShouldSetD3D12ResourceName(Buffer->ResourceLocation))
 	{
-		FD3D12Buffer::FLinkedObjectIterator BufferIt(Buffer);
-
 		if (GNumExplicitGPUsForRendering > 1)
 		{
 			// Generate string of the form "Name (GPU #)" -- assumes GPU index is a single digit.  This is called many times
@@ -851,24 +852,21 @@ void FD3D12DynamicRHI::RHIBindDebugLabelName(FRHICommandListBase& RHICmdList, FR
 			FMemory::Memcpy(&DebugName[0], Name, NameLength * sizeof(TCHAR));
 			FMemory::Memcpy(&DebugName[NameLength], NameSuffix, NameSuffixLengthWithTerminator * sizeof(TCHAR));
 
-			for (; BufferIt; ++BufferIt)
+			for (FD3D12Buffer::FLinkedObjectIterator BufferIt(Buffer); BufferIt; ++BufferIt)
 			{
 				FD3D12Resource* Resource = BufferIt->GetResource();
 
 				DebugName[GPUIndexOffset] = TEXT('0') + BufferIt->GetParentDevice()->GetGPUIndex();
 
-				SetName(Resource, DebugName);
+				SetD3D12ResourceName(BufferIt->GetResource(), DebugName);
 			}
 		}
 		else
 		{
-			SetName(Buffer->GetResource(), Name);
+			SetD3D12ResourceName(Buffer->GetResource(), Name);
 		}
 	}
 #endif
-
-	// Also set on RHI object
-	BufferRHI->SetName(Name);
 }
 
 void FD3D12CommandContext::RHICopyBufferRegion(FRHIBuffer* DestBufferRHI, uint64 DstOffset, FRHIBuffer* SourceBufferRHI, uint64 SrcOffset, uint64 NumBytes)
diff --git a/Engine/Source/Runtime/D3D12RHI/Private/D3D12CommandList.cpp b/Engine/Source/Runtime/D3D12RHI/Private/D3D12CommandList.cpp
--- a/Engine/Source/Runtime/D3D12RHI/Private/D3D12CommandList.cpp
+++ b/Engine/Source/Runtime/D3D12RHI/Private/D3D12CommandList.cpp
@@ -214,9 +214,9 @@ FD3D12CommandList::FD3D12CommandList(FD3D12CommandAllocator* CommandAllocator, F
 	Interfaces.IntelCommandListHandle = UE::RHICore::Intel::GPUCrashDumps::D3D12::RegisterCommandList( Interfaces.GraphicsCommandList );
 #endif
 
-#if NAME_OBJECTS
-	FString Name = FString::Printf(TEXT("FD3D12CommandList (GPU %d)"), Device->GetGPUIndex());
-	SetName(Interfaces.CommandList, Name.GetCharArray().GetData());
+#if RHI_USE_RESOURCE_DEBUG_NAME
+	const FString Name = FString::Printf(TEXT("FD3D12CommandList (GPU %d)"), Device->GetGPUIndex());
+	SetD3D12ObjectName(Interfaces.CommandList, GetData(Name));
 #endif
 
 	D3DX12Residency::Open(ResidencySet);
diff --git a/Engine/Source/Runtime/D3D12RHI/Private/D3D12Descriptors.cpp b/Engine/Source/Runtime/D3D12RHI/Private/D3D12Descriptors.cpp
--- a/Engine/Source/Runtime/D3D12RHI/Private/D3D12Descriptors.cpp
+++ b/Engine/Source/Runtime/D3D12RHI/Private/D3D12Descriptors.cpp
@@ -484,7 +484,7 @@ static FD3D12DescriptorHeap* CreateDescriptorHeap(FD3D12Device* Device, const TC
 	TRefCountPtr<ID3D12DescriptorHeap> Heap;
 	VERIFYD3D12RESULT(Device->GetDevice()->CreateDescriptorHeap(&Desc, IID_PPV_ARGS(Heap.GetInitReference())));
 
-	SetName(Heap, DebugName);
+	SetD3D12ObjectName(Heap, DebugName);
 
 	return new FD3D12DescriptorHeap(Device, MoveTemp(Heap), NumDescriptors, HeapType, Flags, bIsGlobal);
 }
diff --git a/Engine/Source/Runtime/D3D12RHI/Private/D3D12ExplicitDescriptorCache.cpp b/Engine/Source/Runtime/D3D12RHI/Private/D3D12ExplicitDescriptorCache.cpp
--- a/Engine/Source/Runtime/D3D12RHI/Private/D3D12ExplicitDescriptorCache.cpp
+++ b/Engine/Source/Runtime/D3D12RHI/Private/D3D12ExplicitDescriptorCache.cpp
@@ -129,7 +129,7 @@ FD3D12ExplicitDescriptorHeapCache::FEntry FD3D12ExplicitDescriptorHeapCache::All
 	UE_LOG(LogD3D12RHI, Log, TEXT("Creating %s with %d entries"), HeapName, NumDescriptors);
 
 	VERIFYD3D12RESULT(GetParentDevice()->GetDevice()->CreateDescriptorHeap(&Desc, IID_PPV_ARGS(&D3D12Heap)));
-	SetName(D3D12Heap, Desc.Type == D3D12_DESCRIPTOR_HEAP_TYPE_CBV_SRV_UAV ? L"Explicit View Heap" : L"Explicit Sampler Heap");
+	SetD3D12ObjectName(D3D12Heap, Desc.Type == D3D12_DESCRIPTOR_HEAP_TYPE_CBV_SRV_UAV ? L"Explicit View Heap" : L"Explicit Sampler Heap");
 
 	Result.NumDescriptors = NumDescriptors;
 	Result.Type = Type;
diff --git a/Engine/Source/Runtime/D3D12RHI/Private/D3D12Query.cpp b/Engine/Source/Runtime/D3D12RHI/Private/D3D12Query.cpp
--- a/Engine/Source/Runtime/D3D12RHI/Private/D3D12Query.cpp
+++ b/Engine/Source/Runtime/D3D12RHI/Private/D3D12Query.cpp
@@ -87,7 +87,7 @@ FD3D12QueryHeap::FD3D12QueryHeap(FD3D12Device* Device, D3D12_QUERY_TYPE QueryTyp
 		QueryHeapDesc.NodeMask = GetGPUMask().GetNative();
 
 		VERIFYD3D12RESULT(Device->GetDevice()->CreateQueryHeap(&QueryHeapDesc, IID_PPV_ARGS(D3DQueryHeap.GetInitReference())));
-		SetName(D3DQueryHeap, QueryHeapName);
+		SetD3D12ObjectName(D3DQueryHeap, QueryHeapName);
 
 #if ENABLE_RESIDENCY_MANAGEMENT && 0 // Temporary workaround for missing resource usage tracking for query heap
 		D3DX12Residency::Initialize(ResidencyHandle, D3DQueryHeap, GetResultSize() * NumQueries, this);
@@ -114,7 +114,7 @@ FD3D12QueryHeap::FD3D12QueryHeap(FD3D12Device* Device, D3D12_QUERY_TYPE QueryTyp
 			ResultBuffer.GetInitReference(),
 			ResultBufferName
 		));
-		SetName(ResultBuffer, ResultBufferName);
+		SetD3D12ResourceName(ResultBuffer, ResultBufferName);
 	}
 
 	// Map the readback buffer. Resources in a readback heap are allowed
diff --git a/Engine/Source/Runtime/D3D12RHI/Private/D3D12RayTracing.cpp b/Engine/Source/Runtime/D3D12RHI/Private/D3D12RayTracing.cpp
--- a/Engine/Source/Runtime/D3D12RHI/Private/D3D12RayTracing.cpp
+++ b/Engine/Source/Runtime/D3D12RHI/Private/D3D12RayTracing.cpp
@@ -741,7 +741,7 @@ FD3D12RayTracingCompactionRequestHandler::FD3D12RayTracingCompactionRequestHandl
 		ED3D12ResourceStateMode::MultiState,
 		D3D12_RESOURCE_STATE_UNORDERED_ACCESS
 	);
-	SetName(PostBuildInfoBuffer->GetResource(), TEXT("PostBuildInfoBuffer"));
+	SetD3D12ResourceName(PostBuildInfoBuffer->GetResource(), TEXT("PostBuildInfoBuffer"));
 
 	PostBuildInfoStagingBuffer = RHICreateStagingBuffer();
 }
diff --git a/Engine/Source/Runtime/D3D12RHI/Private/D3D12Resources.cpp b/Engine/Source/Runtime/D3D12RHI/Private/D3D12Resources.cpp
--- a/Engine/Source/Runtime/D3D12RHI/Private/D3D12Resources.cpp
+++ b/Engine/Source/Runtime/D3D12RHI/Private/D3D12Resources.cpp
@@ -675,7 +675,7 @@ void FD3D12Heap::SetHeap(ID3D12Heap* HeapIn, const TCHAR* const InName, bool bIn
 	HeapName = InName;
 	HeapDesc = Heap->GetDesc();
 
-	SetName(HeapIn, InName);
+	SetD3D12ObjectName(HeapIn, InName);
 
 #if ENABLE_RESIDENCY_MANAGEMENT
 	bRequiresResidencyTracking = IsGPUOnly(HeapDesc.Properties.Type, &HeapDesc.Properties);
@@ -824,7 +824,7 @@ HRESULT FD3D12Adapter::CreateCommittedResource(const FD3D12ResourceDesc& InDesc,
 		(*ppOutResource)->AddRef();
 
 		// Set a default name (can override later).
-		SetName(*ppOutResource, Name);
+		SetD3D12ResourceName(*ppOutResource, Name);
 
 		(*ppOutResource)->StartTrackingForResidency();
 		
@@ -919,7 +919,7 @@ HRESULT FD3D12Adapter::CreateReservedResource(const FD3D12ResourceDesc& InDesc,
 		(*ppOutResource)->AddRef();
 
 		// Set a default name (can override later).
-		SetName(*ppOutResource, Name);
+		SetD3D12ResourceName(*ppOutResource, Name);
 		
 		// NOTE: reserved resource residency is not tracked/managed by the engine, so we don't need to call StartTrackingForResidency().
 	}
@@ -1028,7 +1028,7 @@ HRESULT FD3D12Adapter::CreatePlacedResource(const FD3D12ResourceDesc& InDesc, FD
 		}
 
 		// Set a default name (can override later).
-		SetName(*ppOutResource, Name);
+		SetD3D12ResourceName(*ppOutResource, Name);
 
 		(*ppOutResource)->AddRef();
 	}
diff --git a/Engine/Source/Runtime/D3D12RHI/Private/D3D12Resources.h b/Engine/Source/Runtime/D3D12RHI/Private/D3D12Resources.h
--- a/Engine/Source/Runtime/D3D12RHI/Private/D3D12Resources.h
+++ b/Engine/Source/Runtime/D3D12RHI/Private/D3D12Resources.h
@@ -344,7 +344,7 @@ public:
 		if (DebugName != Name)
 		{
 			DebugName = FName(Name);
-			::SetName(Resource, Name);
+			SetD3D12ObjectName(GetResource(), Name);
 		}
 	}
 
diff --git a/Engine/Source/Runtime/D3D12RHI/Private/D3D12Texture.cpp b/Engine/Source/Runtime/D3D12RHI/Private/D3D12Texture.cpp
--- a/Engine/Source/Runtime/D3D12RHI/Private/D3D12Texture.cpp
+++ b/Engine/Source/Runtime/D3D12RHI/Private/D3D12Texture.cpp
@@ -2736,7 +2736,10 @@ void FD3D12DynamicRHI::RHIBindDebugLabelName(FRHICommandListBase& RHICmdList, FR
 		return;
 	}
 
-#if NAME_OBJECTS
+#if RHI_USE_RESOURCE_DEBUG_NAME
+	// Also set on RHI object
+	TextureRHI->SetName(Name);
+
 	FD3D12Texture::FLinkedObjectIterator BaseTexture(GetD3D12TextureFromRHITexture(TextureRHI));
 
 	if (GNumExplicitGPUsForRendering > 1)
@@ -2765,17 +2768,14 @@ void FD3D12DynamicRHI::RHIBindDebugLabelName(FRHICommandListBase& RHICmdList, FR
 
 			DebugName[GPUIndexOffset] = TEXT('0') + BaseTexture->GetParentDevice()->GetGPUIndex();
 
-			SetName(Resource, DebugName);
+			SetD3D12ResourceName(Resource, DebugName);
 		}
 	}
 	else
 	{
-		SetName(BaseTexture->GetResource(), Name);
+		SetD3D12ResourceName(BaseTexture->GetResource(), Name);
 	}
 #endif
-
-	// Also set on RHI object
-	TextureRHI->SetName(Name);
 }
 
 FD3D12Texture* FD3D12DynamicRHI::CreateTextureFromResource(bool bTextureArray, bool bCubeTexture, EPixelFormat Format, ETextureCreateFlags TexCreateFlags, const FClearValueBinding& ClearValueBinding, ID3D12Resource* Resource)
diff --git a/Engine/Source/Runtime/D3D12RHI/Private/D3D12Util.cpp b/Engine/Source/Runtime/D3D12RHI/Private/D3D12Util.cpp
--- a/Engine/Source/Runtime/D3D12RHI/Private/D3D12Util.cpp
+++ b/Engine/Source/Runtime/D3D12RHI/Private/D3D12Util.cpp
@@ -44,49 +44,62 @@ void FD3D12DynamicRHI::ForEachDevice(ID3D12Device* inDevice, const PerDeviceFunc
 	}
 }
 
-static FString GetUniqueName()
+bool ShouldSetD3D12ResourceName(const FD3D12ResourceLocation& ResourceLocation)
 {
-	static int64 ID = 0;
-	const int64 UniqueID = FPlatformAtomics::InterlockedIncrement(&ID);
-	const FString UniqueName = FString::Printf(TEXT("D3D12RHIObjectUniqueName%lld"), UniqueID);
-	return UniqueName;
+#if RHI_USE_RESOURCE_DEBUG_NAME
+	// only rename the underlying d3d12 resource if it's not sub allocated (requires resource state tracking or stand alone allocated)
+	return ResourceLocation.GetResource() && (ResourceLocation.GetResource()->RequiresResourceStateTracking() || ResourceLocation.GetType() == FD3D12ResourceLocation::ResourceLocationType::eStandAlone);
+#else
+	return false;
+#endif
 }
 
-void SetName(ID3D12Object* const Object, const TCHAR* const Name)
+void SetD3D12ObjectName(ID3D12Object* Object, const TCHAR* Name)
 {
-#if NAME_OBJECTS
+#if RHI_USE_RESOURCE_DEBUG_NAME
 	if (Object && Name)
 	{
 		VERIFYD3D12RESULT(Object->SetName(Name));
 	}
-	else if (Object)
-	{
-		VERIFYD3D12RESULT(Object->SetName(*GetUniqueName()));
-	}
-#else
-	UNREFERENCED_PARAMETER(Object);
-	UNREFERENCED_PARAMETER(Name);
 #endif
 }
 
-void SetName(FD3D12Resource* const Resource, const TCHAR* const Name)
+void SetD3D12ResourceName(FD3D12Resource* Resource, const TCHAR* Name)
 {
-#if NAME_OBJECTS
-	// Special case for FD3D12Resources because we also store the name as a member in FD3D12Resource
-	if (Resource && Name)
+#if RHI_USE_RESOURCE_DEBUG_NAME
+	if (Resource)
 	{
 		Resource->SetName(Name);
 	}
-	else if (Resource)
+#endif
+}
+
+void SetD3D12ResourceName(FD3D12ResourceLocation& ResourceLocation, const TCHAR* Name)
+{
+#if RHI_USE_RESOURCE_DEBUG_NAME
+	if (ShouldSetD3D12ResourceName(ResourceLocation))
 	{
-		Resource->SetName(*GetUniqueName());
+		SetD3D12ResourceName(ResourceLocation.GetResource(), Name);
 	}
-#else
-	UNREFERENCED_PARAMETER(Resource);
-	UNREFERENCED_PARAMETER(Name);
 #endif
 }
 
+void SetD3D12ResourceName(FD3D12Buffer* Buffer, const TCHAR* Name)
+{
+	if (Buffer)
+	{
+		SetD3D12ResourceName(Buffer->ResourceLocation, Name);
+	}
+}
+
+void SetD3D12ResourceName(FD3D12Texture* Texture, const TCHAR* Name)
+{
+	if (Texture)
+	{
+		SetD3D12ResourceName(Texture->ResourceLocation, Name);
+	}
+}
+
 FString GetD312ObjectName(ID3D12Object* const Object)
 {
 #if NAME_OBJECTS
diff --git a/Engine/Source/Runtime/D3D12RHI/Private/D3D12Util.h b/Engine/Source/Runtime/D3D12RHI/Private/D3D12Util.h
--- a/Engine/Source/Runtime/D3D12RHI/Private/D3D12Util.h
+++ b/Engine/Source/Runtime/D3D12RHI/Private/D3D12Util.h
@@ -49,10 +49,20 @@ namespace D3D12RHI
 
 using namespace D3D12RHI;
 
+class FD3D12Buffer;
 class FD3D12Resource;
+class FD3D12ResourceLocation;
+class FD3D12Texture;
+
+bool ShouldSetD3D12ResourceName(const FD3D12ResourceLocation& ResourceLocation);
+
+void SetD3D12ObjectName(ID3D12Object* Object, const TCHAR* Name);
+
+void SetD3D12ResourceName(FD3D12ResourceLocation& ResourceLocation, const TCHAR* Name);
+void SetD3D12ResourceName(FD3D12Resource* Resource, const TCHAR* Name);
+void SetD3D12ResourceName(FD3D12Buffer* Resource, const TCHAR* Name);
+void SetD3D12ResourceName(FD3D12Texture* Resource, const TCHAR* Name);
 
-void SetName(ID3D12Object* const Object, const TCHAR* const Name);
-void SetName(FD3D12Resource* const Resource, const TCHAR* const Name);
 FString GetD312ObjectName(ID3D12Object* const Object);
 
 enum EShaderVisibility
diff --git a/Engine/Source/Runtime/D3D12RHI/Private/D3D12Viewport.cpp b/Engine/Source/Runtime/D3D12RHI/Private/D3D12Viewport.cpp
--- a/Engine/Source/Runtime/D3D12RHI/Private/D3D12Viewport.cpp
+++ b/Engine/Source/Runtime/D3D12RHI/Private/D3D12Viewport.cpp
@@ -187,7 +187,7 @@ FD3D12Texture* GetSwapChainSurface(FD3D12Device* Parent, EPixelFormat PixelForma
 		return NewTexture;
 	});
 
-	SetName(SwapChainTexture->GetResource(), *Name);
+	SetD3D12ResourceName(SwapChainTexture->GetResource(), *Name);
 
 	const D3D12_RESOURCE_ALLOCATION_INFO AllocationInfo = Parent->GetDevice()->GetResourceAllocationInfo(0, 1, &SwapChainTexture->GetResource()->GetDesc());
 	SwapChainTexture->ResourceLocation.SetSize(AllocationInfo.SizeInBytes);
diff --git a/Engine/Source/Runtime/D3D12RHI/Private/Windows/WindowsD3D12PipelineState.cpp b/Engine/Source/Runtime/D3D12RHI/Private/Windows/WindowsD3D12PipelineState.cpp
--- a/Engine/Source/Runtime/D3D12RHI/Private/Windows/WindowsD3D12PipelineState.cpp
+++ b/Engine/Source/Runtime/D3D12RHI/Private/Windows/WindowsD3D12PipelineState.cpp
@@ -556,7 +556,7 @@ void FD3D12PipelineStateCache::Init(FString& GraphicsCacheFileName, FString& Com
 				VERIFYD3D12RESULT(pDevice1->CreatePipelineLibrary(nullptr, 0, IID_PPV_ARGS(PipelineLibrary.GetInitReference())));
 			}
 
-			SetName(PipelineLibrary, L"Pipeline Library");
+			SetD3D12ObjectName(PipelineLibrary, L"Pipeline Library");
 		}
 	}
 }
diff --git a/Engine/Source/Runtime/OpenGLDrv/Private/OpenGLResources.h b/Engine/Source/Runtime/OpenGLDrv/Private/OpenGLResources.h
--- a/Engine/Source/Runtime/OpenGLDrv/Private/OpenGLResources.h
+++ b/Engine/Source/Runtime/OpenGLDrv/Private/OpenGLResources.h
@@ -870,12 +870,12 @@ public:
 	uint32 FrameLastRendered;
 
 #if GLDEBUG_LABELS_ENABLED
-	FAnsiCharArray TextureDebugName;
-	void SetDebugLabelName(const FAnsiCharArray& TextureDebugNameIn) { TextureDebugName = TextureDebugNameIn; }
+	FAnsiString TextureDebugName;
+	void SetDebugLabelName(const FAnsiString& TextureDebugNameIn) { TextureDebugName = TextureDebugNameIn; }
 	void SetDebugLabelName(const ANSICHAR * TextureDebugNameIn) { TextureDebugName.Append(TextureDebugNameIn, FCStringAnsi::Strlen(TextureDebugNameIn) + 1); }
-	FAnsiCharArray& GetDebugLabelName() { return TextureDebugName; }
+	FAnsiString& GetDebugLabelName() { return TextureDebugName; }
 #else
-	void SetDebugLabelName(FAnsiCharArray TextureDebugNameIn) { checkNoEntry(); }
+	void SetDebugLabelName(const FAnsiString& TextureDebugNameIn) { checkNoEntry(); }
 	FAnsiCharArray& GetDebugLabelName() { checkNoEntry(); static FAnsiCharArray Dummy;  return Dummy; }
 #endif
 
diff --git a/Engine/Source/Runtime/OpenGLDrv/Private/OpenGLTexture.cpp b/Engine/Source/Runtime/OpenGLDrv/Private/OpenGLTexture.cpp
--- a/Engine/Source/Runtime/OpenGLDrv/Private/OpenGLTexture.cpp
+++ b/Engine/Source/Runtime/OpenGLDrv/Private/OpenGLTexture.cpp
@@ -1290,10 +1290,9 @@ void FOpenGLTexture::RestoreEvictedGLResource(bool bAttemptToRetainMips)
 	}
 
 #if GLDEBUG_LABELS_ENABLED
-	FAnsiCharArray& TextureDebugName = EvictionParamsPtr->GetDebugLabelName();
-	if(TextureDebugName.Num())
+	if (FAnsiString& TextureDebugName = EvictionParamsPtr->GetDebugLabelName(); TextureDebugName.Len())
 	{
-		FOpenGL::LabelObject(GL_TEXTURE, this->GetRawResourceName(), TextureDebugName.GetData());
+		SetOpenGLResourceName(this, *TextureDebugName);
 		if (RetainMips == 0)
 		{
 			TextureDebugName.Empty();
@@ -1835,22 +1834,15 @@ void FOpenGLDynamicRHI::InvalidateUAVResourceInCache(GLuint Resource)
 
 void FOpenGLDynamicRHI::RHIBindDebugLabelName(FRHICommandListBase& RHICmdList, FRHITexture* TextureRHI, const TCHAR* Name)
 {
+	TextureRHI->SetName(Name);
+
 #if GLDEBUG_LABELS_ENABLED
-	FAnsiCharArray TextureDebugName;
-	TextureDebugName.Append(TCHAR_TO_ANSI(Name), FCString::Strlen(Name) + 1);
-	RHICmdList.EnqueueLambda([TextureRHI, TextureDebugName = MoveTemp(TextureDebugName)] (FRHICommandListBase& RHICmdList)
-	{
-		VERIFY_GL_SCOPE();
-		FOpenGLTexture* Texture = ResourceCast(TextureRHI);
-		if (Texture->IsEvicted())
+	RHICmdList.EnqueueLambda(
+		[TextureRHI, Name = FAnsiString(Name)] (FRHICommandListBase& RHICmdList)
 		{
-			Texture->EvictionParamsPtr->SetDebugLabelName(TextureDebugName);
+			SetOpenGLResourceName(ResourceCast(TextureRHI), *Name);
 		}
-		else
-		{
-			FOpenGL::LabelObject(GL_TEXTURE, Texture->GetResource(), TextureDebugName.GetData());
-		}
-	});
+	);
 #endif
 }
 
diff --git a/Engine/Source/Runtime/OpenGLDrv/Private/OpenGLUtil.cpp b/Engine/Source/Runtime/OpenGLDrv/Private/OpenGLUtil.cpp
--- a/Engine/Source/Runtime/OpenGLDrv/Private/OpenGLUtil.cpp
+++ b/Engine/Source/Runtime/OpenGLDrv/Private/OpenGLUtil.cpp
@@ -10,6 +10,30 @@
 #include "OpenGLDrvPrivate.h"
 #include "RHICoreStats.h"
 
+void SetOpenGLResourceName(FOpenGLTexture* Texture, const ANSICHAR* Name)
+{
+#if GLDEBUG_LABELS_ENABLED
+	VERIFY_GL_SCOPE();
+
+	if (Texture->IsEvicted())
+	{
+		Texture->EvictionParamsPtr->SetDebugLabelName(Name);
+	}
+	else
+	{
+		FOpenGL::LabelObject(GL_TEXTURE, Texture->GetResource(), Name);
+	}
+#endif
+}
+
+void SetOpenGLResourceName(FOpenGLTexture* Texture, const TCHAR* Name)
+{
+#if GLDEBUG_LABELS_ENABLED
+	const FAnsiString AnsiName(Name);
+	SetOpenGLResourceName(Texture, *AnsiName);
+#endif
+}
+
 #if ENABLE_DEBUG_OUTPUT
 
 // Enable GL debug output
diff --git a/Engine/Source/Runtime/OpenGLDrv/Public/OpenGLUtil.h b/Engine/Source/Runtime/OpenGLDrv/Public/OpenGLUtil.h
--- a/Engine/Source/Runtime/OpenGLDrv/Public/OpenGLUtil.h
+++ b/Engine/Source/Runtime/OpenGLDrv/Public/OpenGLUtil.h
@@ -10,6 +10,8 @@
 #include "OpenGLThirdParty.h"
 #include "OpenGLShaderResources.h"
 
+class FOpenGLTexture;
+
 /** Set to 1 to enable the VERIFY_GL macros which call glGetError */
 #define ENABLE_VERIFY_GL (0 & DO_CHECK)
 #define ENABLE_VERIFY_GL_TRACE 0
@@ -42,6 +44,9 @@ int32 GetOGLDebugOutputLevel();
 
 #define OPENGL_PERFORMANCE_DATA_INVALID (ENABLE_VERIFY_GL | ENABLE_UNIFORM_BUFFER_LAYOUT_VERIFICATION | DEBUG_GL_SHADERS)
 
+void SetOpenGLResourceName(FOpenGLTexture* Texture, const ANSICHAR* Name);
+void SetOpenGLResourceName(FOpenGLTexture* Texture, const TCHAR* Name);
+
 /**
 * Convert from ECubeFace to GLenum type
 * @param Face - ECubeFace type to convert
diff --git a/Engine/Source/Runtime/RHI/Private/DynamicRHI.cpp b/Engine/Source/Runtime/RHI/Private/DynamicRHI.cpp
--- a/Engine/Source/Runtime/RHI/Private/DynamicRHI.cpp
+++ b/Engine/Source/Runtime/RHI/Private/DynamicRHI.cpp
@@ -891,3 +891,31 @@ RHI_API FShaderResourceViewInitializer::FShaderResourceViewInitializer(FRHIBuffe
 {
 	SetTypeFromBuffer(InBuffer);
 }
+
+void FDynamicRHI::RHIBindDebugLabelName(FRHICommandListBase& RHICmdList, FRHITexture* Texture, const TCHAR* Name)
+{
+#if RHI_USE_RESOURCE_DEBUG_NAME
+	Texture->SetName(Name);
+#endif
+}
+
+void FDynamicRHI::RHIBindDebugLabelName(FRHICommandListBase& RHICmdList, FRHIBuffer* Buffer, const TCHAR* Name)
+{
+#if RHI_USE_RESOURCE_DEBUG_NAME
+	Buffer->SetName(Name);
+#endif
+}
+
+void FDynamicRHI::RHIBindDebugLabelName(FRHICommandListBase& RHICmdList, FRHIUnorderedAccessView* UnorderedAccessViewRHI, const TCHAR* Name)
+{
+#if RHI_USE_RESOURCE_DEBUG_NAME
+	if (UnorderedAccessViewRHI->IsTexture())
+	{
+		RHICmdList.BindDebugLabelName(UnorderedAccessViewRHI->GetTexture(), Name);
+	}
+	else
+	{
+		RHICmdList.BindDebugLabelName(UnorderedAccessViewRHI->GetBuffer(), Name);
+	}
+#endif
+}
diff --git a/Engine/Source/Runtime/RHI/Public/DynamicRHI.h b/Engine/Source/Runtime/RHI/Public/DynamicRHI.h
--- a/Engine/Source/Runtime/RHI/Public/DynamicRHI.h
+++ b/Engine/Source/Runtime/RHI/Public/DynamicRHI.h
@@ -606,11 +606,9 @@ public:
 	*/
 	virtual void RHIUpdateTexture3D(FRHICommandListBase& RHICmdList, FRHITexture* Texture, uint32 MipIndex, const struct FUpdateTextureRegion3D& UpdateRegion, uint32 SourceRowPitch, uint32 SourceDepthPitch, const uint8* SourceData) = 0;
 
-
-	// FlushType: Thread safe
-	virtual void RHIBindDebugLabelName(FRHICommandListBase& RHICmdList,FRHITexture* Texture, const TCHAR* Name) = 0;
-	virtual void RHIBindDebugLabelName(FRHICommandListBase& RHICmdList,FRHIBuffer* Buffer, const TCHAR* Name) {}
-	virtual void RHIBindDebugLabelName(FRHICommandListBase& RHICmdList,FRHIUnorderedAccessView* UnorderedAccessViewRHI, const TCHAR* Name) {}
+	RHI_API virtual void RHIBindDebugLabelName(FRHICommandListBase& RHICmdList,FRHITexture* Texture, const TCHAR* Name);
+	RHI_API virtual void RHIBindDebugLabelName(FRHICommandListBase& RHICmdList, FRHIBuffer* Buffer, const TCHAR* Name);
+	RHI_API virtual void RHIBindDebugLabelName(FRHICommandListBase& RHICmdList, FRHIUnorderedAccessView* UnorderedAccessViewRHI, const TCHAR* Name);
 
 	/**
 	* Reads the contents of a texture to an output buffer (non MSAA and MSAA) and returns it as a FColor array.
diff --git a/Engine/Source/Runtime/VulkanRHI/Private/VulkanDynamicRHI.h b/Engine/Source/Runtime/VulkanRHI/Private/VulkanDynamicRHI.h
--- a/Engine/Source/Runtime/VulkanRHI/Private/VulkanDynamicRHI.h
+++ b/Engine/Source/Runtime/VulkanRHI/Private/VulkanDynamicRHI.h
@@ -202,7 +202,6 @@ public:
 		InternalUpdateTexture3D(RHICmdList, Texture, MipIndex, UpdateRegion, SourceRowPitch, SourceDepthPitch, SourceData);
 	}
 	virtual void RHIBindDebugLabelName(FRHICommandListBase& RHICmdList, FRHITexture* Texture, const TCHAR* Name) final override;
-	virtual void RHIBindDebugLabelName(FRHICommandListBase& RHICmdList, FRHIUnorderedAccessView* UnorderedAccessViewRHI, const TCHAR* Name) final override;
 	virtual void RHIReadSurfaceData(FRHITexture* Texture, FIntRect Rect, TArray<FColor>& OutData, FReadSurfaceDataFlags InFlags) final override;
 	virtual void RHIReadSurfaceData(FRHITexture* Texture, FIntRect Rect, TArray<FLinearColor>& OutData, FReadSurfaceDataFlags InFlags) final override;
 	virtual void RHIMapStagingSurface(FRHITexture* Texture, FRHIGPUFence* Fence, void*& OutData, int32& OutWidth, int32& OutHeight, uint32 GPUIndex = 0) final override;
diff --git a/Engine/Source/Runtime/VulkanRHI/Private/VulkanTexture.cpp b/Engine/Source/Runtime/VulkanRHI/Private/VulkanTexture.cpp
--- a/Engine/Source/Runtime/VulkanRHI/Private/VulkanTexture.cpp
+++ b/Engine/Source/Runtime/VulkanRHI/Private/VulkanTexture.cpp
@@ -2220,52 +2220,10 @@ bool FVulkanTexture::GetTextureResourceInfo(FRHIResourceInfo& OutResourceInfo) c
 
 void FVulkanDynamicRHI::RHIBindDebugLabelName(FRHICommandListBase& RHICmdList, FRHITexture* TextureRHI, const TCHAR* Name)
 {
-#if VULKAN_ENABLE_IMAGE_TRACKING_LAYER
-	{
-		FVulkanTexture* VulkanTexture = ResourceCast(TextureRHI);
-		VulkanRHI::BindDebugLabelName(VulkanTexture->Image, Name);
-	}
-#endif
-
-#if VULKAN_ENABLE_DUMP_LAYER
-	{
-// TODO: this dies in the printf on android. Needs investigation.
-#if !PLATFORM_ANDROID
-		FVulkanTexture* VulkanTexture = ResourceCast(TextureRHI);
-#if VULKAN_ENABLE_DUMP_LAYER
-		VulkanRHI::PrintfBegin(*FString::Printf(TEXT("vkDebugMarkerSetObjectNameEXT(0x%p=%s)\n"), VulkanTexture->Image, Name));
-#endif
-#endif
-	}
-#endif
-
-#if VULKAN_ENABLE_DRAW_MARKERS
-	if (auto* SetDebugName = Device->GetSetDebugName())
-	{
-		FVulkanTexture* VulkanTexture = ResourceCast(TextureRHI);
-		FTCHARToUTF8 Converter(Name);
-		VulkanRHI::SetDebugName(SetDebugName, Device->GetInstanceHandle(), VulkanTexture->Image, Converter.Get());
-	}
-#endif
-	FName DebugName(Name);
-	TextureRHI->SetName(DebugName);
-}
+#if RHI_USE_RESOURCE_DEBUG_NAME
+	TextureRHI->SetName(Name);
 
-void FVulkanDynamicRHI::RHIBindDebugLabelName(FRHICommandListBase& RHICmdList, FRHIUnorderedAccessView* UnorderedAccessViewRHI, const TCHAR* Name)
-{
-#if VULKAN_ENABLE_DUMP_LAYER
-	//if (Device->SupportsDebugMarkers())
-	{
-		//if (FRHITexture* Tex2d = UnorderedAccessViewRHI->GetTexture2D())
-		//{
-		//	FVulkanTexture2D* VulkanTexture = (FVulkanTexture2D*)Tex2d;
-		//	VkDebugMarkerObjectTagInfoEXT Info;
-		//	ZeroVulkanStruct(Info, VK_STRUCTURE_TYPE_DEBUG_MARKER_OBJECT_NAME_INFO_EXT);
-		//	Info.objectType = VK_DEBUG_REPORT_OBJECT_TYPE_IMAGE_EXT;
-		//	Info.object = VulkanTexture->Image;
-		//	vkDebugMarkerSetObjectNameEXT(Device->GetInstanceHandle(), &Info);
-		//}
-	}
+	SetVulkanResourceName(Device, ResourceCast(TextureRHI), Name);
 #endif
 }
 
diff --git a/Engine/Source/Runtime/VulkanRHI/Private/VulkanUtil.cpp b/Engine/Source/Runtime/VulkanRHI/Private/VulkanUtil.cpp
--- a/Engine/Source/Runtime/VulkanRHI/Private/VulkanUtil.cpp
+++ b/Engine/Source/Runtime/VulkanRHI/Private/VulkanUtil.cpp
@@ -29,6 +29,28 @@ FVulkanDynamicRHI*	GVulkanRHI = nullptr;
 
 extern CORE_API bool GIsGPUCrashed;
 
+void SetVulkanResourceName(FVulkanDevice* Device, FVulkanTexture* Texture, const TCHAR* Name)
+{
+#if VULKAN_ENABLE_IMAGE_TRACKING_LAYER
+	VulkanRHI::BindDebugLabelName(Texture->Image, Name);
+#endif
+
+#if VULKAN_ENABLE_DUMP_LAYER
+// TODO: this dies in the printf on android. Needs investigation.
+#if !PLATFORM_ANDROID
+	VulkanRHI::PrintfBegin(*FString::Printf(TEXT("vkDebugMarkerSetObjectNameEXT(0x%p=%s)\n"), Texture->Image, Name));
+#endif
+#endif
+
+#if VULKAN_ENABLE_DRAW_MARKERS
+	if (PFN_vkSetDebugUtilsObjectNameEXT SetDebugName = Device->GetSetDebugName())
+	{
+		FAnsiString AnsiName(Name);
+		VulkanRHI::SetDebugName(SetDebugName, Device->GetInstanceHandle(), Texture->Image, GetData(AnsiName));
+	}
+#endif
+}
+
 #if (RHI_NEW_GPU_PROFILER == 0)
 
 static FString		EventDeepString(TEXT("EventTooDeep"));
diff --git a/Engine/Source/Runtime/VulkanRHI/Public/VulkanUtil.h b/Engine/Source/Runtime/VulkanRHI/Public/VulkanUtil.h
--- a/Engine/Source/Runtime/VulkanRHI/Public/VulkanUtil.h
+++ b/Engine/Source/Runtime/VulkanRHI/Public/VulkanUtil.h
@@ -12,6 +12,7 @@
 #include "VulkanThirdParty.h"
 
 class FVulkanDevice;
+class FVulkanTexture;
 
 namespace VulkanRHI
 {
@@ -31,6 +32,7 @@ namespace VulkanRHI
 #define VERIFYVULKANRESULT(VkFunction)				{ const VkResult ScopedResult = VkFunction; if (ScopedResult != VK_SUCCESS) { VulkanRHI::VerifyVulkanResult(ScopedResult, #VkFunction, __FILE__, __LINE__); }}
 #define VERIFYVULKANRESULT_EXPANDED(VkFunction)		{ const VkResult ScopedResult = VkFunction; if (ScopedResult < VK_SUCCESS) { VulkanRHI::VerifyVulkanResult(ScopedResult, #VkFunction, __FILE__, __LINE__); }}
 
+void SetVulkanResourceName(FVulkanDevice* Device, FVulkanTexture* Texture, const TCHAR* Name);
 
 template<typename T>
 inline bool CopyAndReturnNotEqual(T& A, T B)
diff --git a/Engine/Source/Runtime/Windows/D3D11RHI/Private/D3D11Buffer.cpp b/Engine/Source/Runtime/Windows/D3D11RHI/Private/D3D11Buffer.cpp
--- a/Engine/Source/Runtime/Windows/D3D11RHI/Private/D3D11Buffer.cpp
+++ b/Engine/Source/Runtime/Windows/D3D11RHI/Private/D3D11Buffer.cpp
@@ -136,11 +136,13 @@ void FD3D11DynamicRHI::FinalizeCreateBufferInternal(FD3D11Buffer* Buffer, TConst
 		}
 	}
 
-	if (FName BufferName = Buffer->GetName(); BufferName != NAME_None)
+#if RHI_USE_RESOURCE_DEBUG_NAME
+	if (const FName BufferName = Buffer->GetName(); BufferName != NAME_None)
 	{
 		const FAnsiString BufferNameString(BufferName.ToString());
-		BufferResource->SetPrivateData(WKPDID_D3DDebugObjectName, BufferNameString.Len() + 1, *BufferNameString);
+		SetD3D11ObjectName(BufferResource, BufferNameString);
 	}
+#endif
 
 	Buffer->Resource = BufferResource;
 
@@ -366,14 +368,11 @@ void FD3D11DynamicRHI::RHIReplaceResources(FRHICommandListBase& RHICmdList, TArr
 
 void FD3D11DynamicRHI::RHIBindDebugLabelName(FRHICommandListBase& RHICmdList, FRHIBuffer* BufferRHI, const TCHAR* Name)
 {
-	check(BufferRHI);
-	FName DebugName(Name);
-	BufferRHI->SetName(DebugName);
 #if RHI_USE_RESOURCE_DEBUG_NAME
-	FD3D11Buffer* BufferD3D = ResourceCast(BufferRHI);
-	if (BufferD3D->Resource != nullptr)
-	{
-		BufferD3D->Resource->SetPrivateData(WKPDID_D3DDebugObjectName, FCString::Strlen(Name) + 1, TCHAR_TO_ANSI(Name));
-	}
+	check(BufferRHI);
+
+	BufferRHI->SetName(Name);
+
+	SetD3D11ResourceName(ResourceCast(BufferRHI), Name);
 #endif
-}
\ No newline at end of file
+}
diff --git a/Engine/Source/Runtime/Windows/D3D11RHI/Private/D3D11Texture.cpp b/Engine/Source/Runtime/Windows/D3D11RHI/Private/D3D11Texture.cpp
--- a/Engine/Source/Runtime/Windows/D3D11RHI/Private/D3D11Texture.cpp
+++ b/Engine/Source/Runtime/Windows/D3D11RHI/Private/D3D11Texture.cpp
@@ -609,6 +609,13 @@ void FD3D11DynamicRHI::FinalizeCreateTexture2DInternal(FD3D11Texture* Texture, T
 		SafeCreateTexture2D(Direct3DDevice, Format, &Desc, pSubresourceData, TextureResource.GetInitReference(), *DebugName);
 	}
 
+#if RHI_USE_RESOURCE_DEBUG_NAME
+	if (Texture->GetName() != NAME_None)
+	{
+		SetD3D11ObjectName(TextureResource, DebugName);
+	}
+#endif
+
 	TArray<TRefCountPtr<ID3D11RenderTargetView>> RenderTargetViews;
 	bool bCreatedRTVPerSlice = false;
 
@@ -722,6 +729,13 @@ void FD3D11DynamicRHI::FinalizeCreateTexture3DInternal(FD3D11Texture* Texture, T
 	TRefCountPtr<ID3D11Texture3D> TextureResource;
 	SafeCreateTexture3D(Direct3DDevice, TextureDesc.Format, &Desc, pSubresourceData, TextureResource.GetInitReference(), *DebugName);
 
+#if RHI_USE_RESOURCE_DEBUG_NAME
+	if (Texture->GetName() != NAME_None)
+	{
+		SetD3D11ObjectName(TextureResource, DebugName);
+	}
+#endif
+
 	// Create a shader resource view for the texture.
 	TRefCountPtr<ID3D11ShaderResourceView> ShaderResourceView;
 
@@ -1303,23 +1317,12 @@ void FD3D11DynamicRHI::RHIEndUpdateTexture3D(FRHICommandListBase& RHICmdList, FU
 	UpdateData.Data = nullptr;
 }
 
-/*-----------------------------------------------------------------------------
-	Cubemap texture support.
------------------------------------------------------------------------------*/
-
 void FD3D11DynamicRHI::RHIBindDebugLabelName(FRHICommandListBase& RHICmdList, FRHITexture* TextureRHI, const TCHAR* Name)
 {
-	FD3D11Texture* Texture = ResourceCast(TextureRHI);
-
-	//todo: require names at texture creation time.
-	FName DebugName(Name);
-	Texture->SetName(DebugName);
 #if RHI_USE_RESOURCE_DEBUG_NAME
-	ID3D11Resource* ResourceD3D = Texture->GetResource();
-	if (ResourceD3D != nullptr)
-	{
-		ResourceD3D->SetPrivateData(WKPDID_D3DDebugObjectName, FCString::Strlen(Name) + 1, TCHAR_TO_ANSI(Name));
-	}
+	TextureRHI->SetName(Name);
+
+	SetD3D11ResourceName(ResourceCast(TextureRHI), Name);
 #endif
 }
 
diff --git a/Engine/Source/Runtime/Windows/D3D11RHI/Private/D3D11UAV.cpp b/Engine/Source/Runtime/Windows/D3D11RHI/Private/D3D11UAV.cpp
--- a/Engine/Source/Runtime/Windows/D3D11RHI/Private/D3D11UAV.cpp
+++ b/Engine/Source/Runtime/Windows/D3D11RHI/Private/D3D11UAV.cpp
@@ -134,16 +134,10 @@ FUnorderedAccessViewRHIRef FD3D11DynamicRHI::RHICreateUnorderedAccessView(class
 void FD3D11DynamicRHI::RHIBindDebugLabelName(FRHICommandListBase& RHICmdList, FRHIUnorderedAccessView* UnorderedAccessViewRHI, const TCHAR* Name)
 {
 #if RHI_USE_RESOURCE_DEBUG_NAME
-	FD3D11UnorderedAccessView* UAV = ResourceCast(UnorderedAccessViewRHI);
-	if (UAV->View != nullptr)
-	{
-		UAV->View->SetPrivateData(WKPDID_D3DDebugObjectName, FCString::Strlen(Name) + 1, TCHAR_TO_ANSI(Name));
-	}
+	SetD3D11ResourceName(ResourceCast(UnorderedAccessViewRHI), Name);
 #endif
 }
 
-
-
 // -----------------------------------------------------------------------------------------------------
 //
 //                                         UAV Clear Functions                                          
diff --git a/Engine/Source/Runtime/Windows/D3D11RHI/Private/D3D11Util.cpp b/Engine/Source/Runtime/Windows/D3D11RHI/Private/D3D11Util.cpp
--- a/Engine/Source/Runtime/Windows/D3D11RHI/Private/D3D11Util.cpp
+++ b/Engine/Source/Runtime/Windows/D3D11RHI/Private/D3D11Util.cpp
@@ -15,6 +15,42 @@
 #define D3DERR(x) case x: ErrorCodeText = TEXT(#x); break;
 #define LOCTEXT_NAMESPACE "Developer.MessageLog"
 
+void SetD3D11ObjectName(ID3D11DeviceChild* Object, FAnsiStringView Name)
+{
+#if RHI_USE_RESOURCE_DEBUG_NAME
+	if (Object)
+	{
+		Object->SetPrivateData(WKPDID_D3DDebugObjectName, Name.Len() + 1, GetData(Name));
+	}
+#endif
+}
+
+void SetD3D11ObjectName(ID3D11DeviceChild* Object, FStringView Name)
+{
+#if RHI_USE_RESOURCE_DEBUG_NAME
+	if (Object != nullptr)
+	{
+		FAnsiString AnsiString(Name);
+		SetD3D11ObjectName(Object, AnsiString);
+	}
+#endif
+}
+
+void SetD3D11ResourceName(FD3D11Texture* Texture, FStringView Name)
+{
+	SetD3D11ObjectName(Texture->GetResource(), Name);
+}
+
+void SetD3D11ResourceName(FD3D11Buffer* Buffer, FStringView Name)
+{
+	SetD3D11ObjectName(Buffer->GetResource(), Name);
+}
+
+void SetD3D11ResourceName(FD3D11UnorderedAccessView* UAV, FStringView Name)
+{
+	SetD3D11ObjectName(UAV->View, Name);
+}
+
 static FString GetD3D11DeviceHungErrorString(HRESULT ErrorCode)
 {
 	FString ErrorCodeText;
diff --git a/Engine/Source/Runtime/Windows/D3D11RHI/Public/D3D11Resources.h b/Engine/Source/Runtime/Windows/D3D11RHI/Public/D3D11Resources.h
--- a/Engine/Source/Runtime/Windows/D3D11RHI/Public/D3D11Resources.h
+++ b/Engine/Source/Runtime/Windows/D3D11RHI/Public/D3D11Resources.h
@@ -457,6 +457,11 @@ public:
 	{
 		return FRHIResource::GetRefCount();
 	}
+
+	ID3D11Resource* GetResource() const
+	{
+		return Resource.GetReference();
+	}
 };
 
 class FD3D11StagingBuffer final : public FRHIStagingBuffer
diff --git a/Engine/Source/Runtime/Windows/D3D11RHI/Public/D3D11Util.h b/Engine/Source/Runtime/Windows/D3D11RHI/Public/D3D11Util.h
--- a/Engine/Source/Runtime/Windows/D3D11RHI/Public/D3D11Util.h
+++ b/Engine/Source/Runtime/Windows/D3D11RHI/Public/D3D11Util.h
@@ -10,13 +10,22 @@
 #include "RHIDefinitions.h"
 #include "Templates/RefCounting.h"
 #include "Windows/D3D11ThirdParty.h"
+#include "Containers/StringFwd.h"
 
+class FD3D11Buffer;
+class FD3D11Texture;
+class FD3D11UnorderedAccessView;
 class FRHIBuffer;
 class FRHITexture;
 
 #define D3D11RHI_IMMEDIATE_CONTEXT	(GD3D11RHI->GetDeviceContext())
 #define D3D11RHI_DEVICE				(GD3D11RHI->GetDevice())
 
+void SetD3D11ObjectName(ID3D11DeviceChild* Object, FAnsiStringView Name);
+void SetD3D11ObjectName(ID3D11DeviceChild* Object, FStringView Name);
+void SetD3D11ResourceName(FD3D11Texture* Texture, FStringView Name);
+void SetD3D11ResourceName(FD3D11Buffer* Buffer, FStringView Name);
+void SetD3D11ResourceName(FD3D11UnorderedAccessView* UAV, FStringView Name);
 
 /**
  * Checks that the given result isn't a failure.  If it is, the application does not exit and only logs an appropriate error message.
