Fixed GetSubresourceCopyBox not respecting block sizes

diff --git a/Engine/Source/Runtime/D3D12RHI/Private/D3D12Texture.cpp b/Engine/Source/Runtime/D3D12RHI/Private/D3D12Texture.cpp
--- a/Engine/Source/Runtime/D3D12RHI/Private/D3D12Texture.cpp
+++ b/Engine/Source/Runtime/D3D12RHI/Private/D3D12Texture.cpp
@@ -1328,7 +1328,7 @@ static FSubresourceIndices GetSubresourceIndices(const D3D12_RESOURCE_DESC& Desc
 	};
 }
 
-static D3D12_BOX GetSubresourceCopyBox(const D3D12_RESOURCE_DESC& Desc, const D3D12_TEXTURE_COPY_LOCATION* OtherLocation, const FSubresourceIndices& Indices)
+static D3D12_BOX GetSubresourceCopyBox(const D3D12_RESOURCE_DESC& Desc, const D3D12_TEXTURE_COPY_LOCATION* OtherLocation, const FSubresourceIndices& Indices, const FPixelFormatInfo& PixelInfo)
 {
 	// Dimensions may be derived from the counter parts footprint
 	if (OtherLocation->Type == D3D12_TEXTURE_COPY_TYPE_PLACED_FOOTPRINT)
@@ -1344,9 +1344,9 @@ static D3D12_BOX GetSubresourceCopyBox(const D3D12_RESOURCE_DESC& Desc, const D3
 	// Otherwise just assume the entire (sub)resource
 	return {
 		.left = 0, .top = 0, .front = 0,
-		.right = FMath::Max<uint32>(1u, Desc.Width >> Indices.MipIndex),
-		.bottom = FMath::Max<uint32>(1u, Desc.Height >> Indices.MipIndex),
-		.back = Desc.Dimension == D3D12_RESOURCE_DIMENSION_TEXTURE3D ? FMath::Max<uint32>(1u, Desc.DepthOrArraySize >> Indices.MipIndex) : 1
+		.right = AlignArbitrary<uint32>(FMath::Max<uint32>(1u, Desc.Width >> Indices.MipIndex), PixelInfo.BlockSizeX),
+		.bottom = AlignArbitrary<uint32>(FMath::Max<uint32>(1u, Desc.Height >> Indices.MipIndex), PixelInfo.BlockSizeY),
+		.back = AlignArbitrary<uint32>(Desc.Dimension == D3D12_RESOURCE_DIMENSION_TEXTURE3D ? FMath::Max<uint32>(1u, Desc.DepthOrArraySize >> Indices.MipIndex) : 1, PixelInfo.BlockSizeZ)
 	};
 }
 
@@ -1397,7 +1397,7 @@ FORCENOINLINE void FD3D12CommandContext::CopyTextureRegionChecked(
 		{
 			const FSubresourceIndices Indices = GetSubresourceIndices(SourceDesc, SourceCopyLocation->SubresourceIndex);
 
-			const D3D12_BOX CopyBox = SourceBox ? *SourceBox : GetSubresourceCopyBox(SourceDesc, DestCopyLocation, Indices);
+			const D3D12_BOX CopyBox = SourceBox ? *SourceBox : GetSubresourceCopyBox(SourceDesc, DestCopyLocation, Indices, SourcePixelFormatInfo);
 
 			RHI_BREADCRUMB_CHECK_SHIPPINGF(*this, (CopyBox.left % SourcePixelFormatInfo.BlockSizeX) == 0 && (CopyBox.right % SourcePixelFormatInfo.BlockSizeX) == 0, TEXT("Width not aligned to block size for: '%s'"), *DebugName.ToString());
 			RHI_BREADCRUMB_CHECK_SHIPPINGF(*this, (CopyBox.top % SourcePixelFormatInfo.BlockSizeY) == 0 && (CopyBox.bottom % SourcePixelFormatInfo.BlockSizeY) == 0, TEXT("Height not aligned to block size for: '%s'"), *DebugName.ToString());
@@ -1427,7 +1427,7 @@ FORCENOINLINE void FD3D12CommandContext::CopyTextureRegionChecked(
 		case D3D12_TEXTURE_COPY_TYPE_SUBRESOURCE_INDEX:
 		{
 			const FSubresourceIndices Indices = GetSubresourceIndices(DestDesc, DestCopyLocation->SubresourceIndex);
-			const D3D12_BOX CopyBox = SourceBox ? GetBoxAdjustedForPixelFormat(*SourceBox, SourcePixelFormatInfo, DestPixelFormatInfo) : GetSubresourceCopyBox(DestDesc, SourceCopyLocation, Indices);
+			const D3D12_BOX CopyBox = SourceBox ? GetBoxAdjustedForPixelFormat(*SourceBox, SourcePixelFormatInfo, DestPixelFormatInfo) : GetSubresourceCopyBox(DestDesc, SourceCopyLocation, Indices, DestPixelFormatInfo);
 				
 			RHI_BREADCRUMB_CHECK_SHIPPINGF(*this, DestX % DestPixelFormatInfo.BlockSizeX == 0 && (CopyBox.right - CopyBox.left) % DestPixelFormatInfo.BlockSizeX == 0, TEXT("Width not aligned to block size for: '%s'"), *DebugName.ToString());
 			RHI_BREADCRUMB_CHECK_SHIPPINGF(*this, DestY % DestPixelFormatInfo.BlockSizeY == 0 && (CopyBox.bottom - CopyBox.top) % DestPixelFormatInfo.BlockSizeY == 0, TEXT("Height not aligned to block size for: '%s'"), *DebugName.ToString());
