Add RHIUpdateResourceResidency API to ID3D12DynamicRHI
- Allows external plugins to add resources to the current context's residency set, to ensure those resources are made resident before the work runs on the GPU.
- Previously external plugins had no way to ensure resources are resident, which led to GPU page faults.

diff --git a/Engine/Source/Runtime/D3D12RHI/Private/D3D12RHI.cpp b/Engine/Source/Runtime/D3D12RHI/Private/D3D12RHI.cpp
--- a/Engine/Source/Runtime/D3D12RHI/Private/D3D12RHI.cpp
+++ b/Engine/Source/Runtime/D3D12RHI/Private/D3D12RHI.cpp
@@ -789,6 +789,20 @@ void FD3D12DynamicRHI::RHIFlushResourceBarriers(FRHICommandListBase& RHICmdList,
 	FD3D12CommandContext::Get(RHICmdList, InGPUIndex).FlushResourceBarriers();
 }
 
+void FD3D12DynamicRHI::RHIUpdateResourceResidency(FRHICommandListBase& RHICmdList, uint32 InGPUIndex, FRHIResource* InResource)
+{
+	if (InResource->GetType() == RRT_Buffer)
+	{
+		FD3D12Buffer* D3D12Buffer = ResourceCast(static_cast<FRHIBuffer*>(InResource), InGPUIndex);
+		FD3D12CommandContext::Get(RHICmdList, InGPUIndex).UpdateResidency(D3D12Buffer->GetResource());
+	}
+	else if (InResource->GetType() == RRT_Texture || InResource->GetType() == RRT_TextureReference)
+	{
+		FD3D12Texture* D3D12Texture = GetD3D12TextureFromRHITexture(static_cast<FRHITexture*>(InResource), InGPUIndex);
+		FD3D12CommandContext::Get(RHICmdList, InGPUIndex).UpdateResidency(D3D12Texture->GetResource());
+	}
+}
+
 void* FD3D12DynamicRHI::RHIGetNativeDevice()
 {
 	return (void*)GetAdapter().GetD3DDevice();
diff --git a/Engine/Source/Runtime/D3D12RHI/Private/D3D12RHIPrivate.h b/Engine/Source/Runtime/D3D12RHI/Private/D3D12RHIPrivate.h
--- a/Engine/Source/Runtime/D3D12RHI/Private/D3D12RHIPrivate.h
+++ b/Engine/Source/Runtime/D3D12RHI/Private/D3D12RHIPrivate.h
@@ -422,6 +422,10 @@ public:
 	virtual void RHIRunOnQueue(ED3D12RHIRunOnQueueType QueueType, TFunction<void(ID3D12CommandQueue*)>&& CodeToRun, bool bWaitForSubmission) final override;
 	virtual void RHIFlushResourceBarriers(FRHICommandListBase& RHICmdList, uint32 InGPUIndex) final override;
 
+	// Allows external plugins to add the given FRHIResource's residency handle to the current RHICmdList's context.
+	// This ensures resources used by external plugins and made resident before the work begins on the GPU.
+	virtual void RHIUpdateResourceResidency(FRHICommandListBase& RHICmdList, uint32 InGPUIndex, FRHIResource* InResource) final override;
+
 	virtual IRHITransientResourceAllocator* RHICreateTransientResourceAllocator() override;
 
 	virtual void RHIWriteGPUFence_TopOfPipe(FRHICommandListBase& RHICmdList, FRHIGPUFence* FenceRHI) final override;
diff --git a/Engine/Source/Runtime/D3D12RHI/Public/ID3D12DynamicRHI.h b/Engine/Source/Runtime/D3D12RHI/Public/ID3D12DynamicRHI.h
--- a/Engine/Source/Runtime/D3D12RHI/Public/ID3D12DynamicRHI.h
+++ b/Engine/Source/Runtime/D3D12RHI/Public/ID3D12DynamicRHI.h
@@ -59,6 +59,7 @@ struct ID3D12DynamicRHI : public FDynamicRHI
 	virtual void				  RHIRunOnQueue(ED3D12RHIRunOnQueueType QueueType, TFunction<void(ID3D12CommandQueue*)>&& CodeToRun, bool bWaitForSubmission) = 0;
 
 	virtual void				  RHIFlushResourceBarriers(FRHICommandListBase& RHICmdList, uint32 InGPUIndex) = 0;
+	virtual void				  RHIUpdateResourceResidency(FRHICommandListBase& RHICmdList, uint32 InGPUIndex, FRHIResource* InResource) = 0;
 };
 
 inline bool IsRHID3D12()
