Add SyncPoint debug names information to display in Unreal Insight in order to facilitate profiling

diff --git a/Engine/Source/Runtime/D3D12RHI/Private/D3D12CommandContext.cpp b/Engine/Source/Runtime/D3D12RHI/Private/D3D12CommandContext.cpp
--- a/Engine/Source/Runtime/D3D12RHI/Private/D3D12CommandContext.cpp
+++ b/Engine/Source/Runtime/D3D12RHI/Private/D3D12CommandContext.cpp
@@ -587,7 +587,7 @@ void FD3D12QueryAllocator::CloseAndReset(TMap<TRefCountPtr<FD3D12QueryHeap>, TAr
 
 FD3D12CopyScope::FD3D12CopyScope(FD3D12Device* Device, ED3D12SyncPointType SyncPointType, FD3D12SyncPointRef const& WaitSyncPoint)
 	: Device(Device)
-	, SyncPoint(FD3D12SyncPoint::Create(SyncPointType))
+	, SyncPoint(FD3D12SyncPoint::Create(SyncPointType, TEXT("CopyScope")))
 	, Context(*Device->ObtainContextCopy())
 {
 	if (WaitSyncPoint)
@@ -642,7 +642,7 @@ void FD3D12ContextCommon::FlushCommands(ED3D12FlushFlags FlushFlags)
 
 	if (EnumHasAnyFlags(FlushFlags, ED3D12FlushFlags::WaitForCompletion))
 	{
-		SyncPoint = FD3D12SyncPoint::Create(ED3D12SyncPointType::GPUAndCPU);
+		SyncPoint = FD3D12SyncPoint::Create(ED3D12SyncPointType::GPUAndCPU, TEXT("WaitForCompletion"));
 		SignalSyncPoint(SyncPoint);
 	}
 
@@ -839,7 +839,7 @@ void FD3D12DynamicRHI::RHICreateTransition(FRHITransition* Transition, const FRH
 			TRHIPipelineArray<FD3D12SyncPointRef>& DeviceSyncPoints = Data->SyncPoints.Emplace_GetRef();
 			for (ERHIPipeline Pipeline : MakeFlagsRange(CreateInfo.SrcPipelines))
 			{
-				DeviceSyncPoints[Pipeline] = FD3D12SyncPoint::Create(ED3D12SyncPointType::GPUOnly);
+				DeviceSyncPoints[Pipeline] = FD3D12SyncPoint::Create(ED3D12SyncPointType::GPUOnly, TEXT("Transition"));
 			}
 		}
 	}
@@ -913,7 +913,7 @@ void FD3D12CommandContextRedirector::RHITransferResources(TConstArrayView<FTrans
 		TStaticArray<FD3D12SyncPointRef, MAX_NUM_GPUS> SyncPoints;
 		for (uint32 GPUIndex : SignalMask)
 		{
-			SyncPoints[GPUIndex] = FD3D12SyncPoint::Create(ED3D12SyncPointType::GPUOnly);
+			SyncPoints[GPUIndex] = FD3D12SyncPoint::Create(ED3D12SyncPointType::GPUOnly, TEXT("TransferResources"));
 			PhysicalContexts[GPUIndex]->SignalSyncPoint(SyncPoints[GPUIndex]);
 		}
 
@@ -1119,7 +1119,7 @@ void FD3D12CommandContextRedirector::RHITransferResourceSignal(TConstArrayView<F
 	uint32 FenceIndex = 0;
 	for (uint32 SrcGPUIndex : SrcGPUMask)
 	{
-		FD3D12SyncPointRef SyncPoint = FD3D12SyncPoint::Create(ED3D12SyncPointType::GPUOnly);
+		FD3D12SyncPointRef SyncPoint = FD3D12SyncPoint::Create(ED3D12SyncPointType::GPUOnly, TEXT("TransferReourceSignal"));
 		SyncPoint->AddRef();
 
 		PhysicalContexts[SrcGPUIndex]->SignalSyncPoint(SyncPoint);
@@ -1228,7 +1228,7 @@ void FD3D12CommandContextRedirector::RHICrossGPUTransfer(TConstArrayView<FTransf
 	}
 
 	// Post-copy synchronization
-	FD3D12SyncPointRef SyncPoint = FD3D12SyncPoint::Create(ED3D12SyncPointType::GPUOnly);
+	FD3D12SyncPointRef SyncPoint = FD3D12SyncPoint::Create(ED3D12SyncPointType::GPUOnly, TEXT("RHICrossGPUTransfer"));
 	PhysicalContexts[Params[0].SrcGPUIndex]->SignalSyncPoint(SyncPoint);
 
 	for (FCrossGPUTransferFence* PostTransferSyncPoint : PostTransfer)
@@ -1243,7 +1243,7 @@ void FD3D12CommandContextRedirector::RHICrossGPUTransferSignal(TConstArrayView<F
 {
 	for (FCrossGPUTransferFence* TransferSyncPoint : PreTransfer)
 	{
-		FD3D12SyncPointRef SyncPoint = FD3D12SyncPoint::Create(ED3D12SyncPointType::GPUOnly);
+		FD3D12SyncPointRef SyncPoint = FD3D12SyncPoint::Create(ED3D12SyncPointType::GPUOnly, TEXT("RHICrossGPUTransferSignal"));
 		SyncPoint->AddRef();
 
 		PhysicalContexts[TransferSyncPoint->SignalGPUIndex]->SignalSyncPoint(SyncPoint);
diff --git a/Engine/Source/Runtime/D3D12RHI/Private/D3D12CommandContext.h b/Engine/Source/Runtime/D3D12RHI/Private/D3D12CommandContext.h
--- a/Engine/Source/Runtime/D3D12RHI/Private/D3D12CommandContext.h
+++ b/Engine/Source/Runtime/D3D12RHI/Private/D3D12CommandContext.h
@@ -265,7 +265,7 @@ public:
 	{
 		if (!ContextSyncPoint)
 		{
-			ContextSyncPoint = FD3D12SyncPoint::Create(ED3D12SyncPointType::GPUAndCPU);
+			ContextSyncPoint = FD3D12SyncPoint::Create(ED3D12SyncPointType::GPUAndCPU, TEXT("ContextSyncPoint"));
 			BatchedSyncPoints.ToSignal.Add(ContextSyncPoint);
 		}
 
diff --git a/Engine/Source/Runtime/D3D12RHI/Private/D3D12Device.cpp b/Engine/Source/Runtime/D3D12RHI/Private/D3D12Device.cpp
--- a/Engine/Source/Runtime/D3D12RHI/Private/D3D12Device.cpp
+++ b/Engine/Source/Runtime/D3D12RHI/Private/D3D12Device.cpp
@@ -562,7 +562,7 @@ void FD3D12Device::BlockUntilIdle()
 
 	for (uint32 QueueTypeIndex = 0; QueueTypeIndex < (uint32)ED3D12QueueType::Count; ++QueueTypeIndex)
 	{
-		FD3D12SyncPointRef SyncPoint = FD3D12SyncPoint::Create(ED3D12SyncPointType::GPUAndCPU);
+		FD3D12SyncPointRef SyncPoint = FD3D12SyncPoint::Create(ED3D12SyncPointType::GPUAndCPU, TEXT("BlockUntilGPUIdle"));
 
 		FD3D12Payload* Payload = new FD3D12Payload(GetQueue((ED3D12QueueType)QueueTypeIndex));
 		Payload->SyncPointsToSignal.Add(SyncPoint);
diff --git a/Engine/Source/Runtime/D3D12RHI/Private/D3D12DirectCommandListManager.cpp b/Engine/Source/Runtime/D3D12RHI/Private/D3D12DirectCommandListManager.cpp
--- a/Engine/Source/Runtime/D3D12RHI/Private/D3D12DirectCommandListManager.cpp
+++ b/Engine/Source/Runtime/D3D12RHI/Private/D3D12DirectCommandListManager.cpp
@@ -61,7 +61,14 @@ void FD3D12DynamicRHI::RHIWriteGPUFence_TopOfPipe(FRHICommandListBase& RHICmdLis
 	for (uint32 GPUIndex : RHICmdList.GetGPUMask())
 	{
 		checkf(Fence->SyncPoints[GPUIndex] == nullptr, TEXT("The fence for the current GPU node has already been issued."));
-		Fence->SyncPoints[GPUIndex] = FD3D12SyncPoint::Create(ED3D12SyncPointType::GPUAndCPU);
+
+		Fence->SyncPoints[GPUIndex] = FD3D12SyncPoint::Create(ED3D12SyncPointType::GPUAndCPU,
+#if RHI_USE_SYNC_POINT_DEBUG_NAME
+		*FenceRHI->GetFName().ToString()
+#else 
+		TEXT("")
+#endif
+		);
 	}
 
 	Fence->NumPendingWriteCommands.Increment();
diff --git a/Engine/Source/Runtime/D3D12RHI/Private/D3D12Query.cpp b/Engine/Source/Runtime/D3D12RHI/Private/D3D12Query.cpp
--- a/Engine/Source/Runtime/D3D12RHI/Private/D3D12Query.cpp
+++ b/Engine/Source/Runtime/D3D12RHI/Private/D3D12Query.cpp
@@ -7,6 +7,7 @@
 #include "D3D12Query.h"
 #include "D3D12RHIPrivate.h"
 #include "ProfilingDebugging/AssetMetadataTrace.h"
+#include "RHIStrings.h"
 
 namespace D3D12RHI
 {
@@ -167,7 +168,7 @@ void FD3D12DynamicRHI::RHIBeginRenderQueryBatch_TopOfPipe(FRHICommandListBase& R
 		auto& QueryBatchData = RHICmdList.GetQueryBatchData(QueryType);
 		checkf(QueryBatchData[GPUIndex] == nullptr, TEXT("A query batch for this type has already begun on this command list."));
 
-		FD3D12SyncPointRef SyncPoint = FD3D12SyncPoint::Create(ED3D12SyncPointType::GPUAndCPU);
+		FD3D12SyncPointRef SyncPoint = FD3D12SyncPoint::Create(ED3D12SyncPointType::GPUAndCPU, GetRenderQueryTypeName(QueryType));
 
 		// Keep a reference on the RHI command list, so we can retrieve it later in BeginQuery/EndQuery/EndBatch.
 		QueryBatchData[GPUIndex] = SyncPoint.GetReference();
@@ -219,7 +220,7 @@ void FD3D12DynamicRHI::RHIEndRenderQuery_TopOfPipe(FRHICommandListBase& RHICmdLi
 		else
 		{
 			// Queries issued outside of a batch use one sync point per query.
-			Query->SyncPoint = FD3D12SyncPoint::Create(ED3D12SyncPointType::GPUAndCPU);
+			Query->SyncPoint = FD3D12SyncPoint::Create(ED3D12SyncPointType::GPUAndCPU, GetRenderQueryTypeName(Query->Type));
 
 			RHICmdList.EnqueueLambda([SyncPoint = Query->SyncPoint, GPUIndex](FRHICommandListBase& ExecutingCmdList) mutable
 			{
diff --git a/Engine/Source/Runtime/D3D12RHI/Private/D3D12RHI.cpp b/Engine/Source/Runtime/D3D12RHI/Private/D3D12RHI.cpp
--- a/Engine/Source/Runtime/D3D12RHI/Private/D3D12RHI.cpp
+++ b/Engine/Source/Runtime/D3D12RHI/Private/D3D12RHI.cpp
@@ -411,7 +411,7 @@ void FD3D12DynamicRHI::EnqueueEndOfPipeTask(TUniqueFunction<void()> TaskFunc, TU
 	{
 		FD3D12Payload* Payload = new FD3D12Payload(Queue);
 
-		FD3D12SyncPointRef SyncPoint = FD3D12SyncPoint::Create(ED3D12SyncPointType::GPUAndCPU);
+		FD3D12SyncPointRef SyncPoint = FD3D12SyncPoint::Create(ED3D12SyncPointType::GPUAndCPU, GetD3DCommandQueueTypeName(Queue.QueueType) );
 		Payload->SyncPointsToSignal.Emplace(SyncPoint);
 		Prereqs.Add(SyncPoint->GetGraphEvent());
 
diff --git a/Engine/Source/Runtime/D3D12RHI/Private/D3D12Submission.cpp b/Engine/Source/Runtime/D3D12RHI/Private/D3D12Submission.cpp
--- a/Engine/Source/Runtime/D3D12RHI/Private/D3D12Submission.cpp
+++ b/Engine/Source/Runtime/D3D12RHI/Private/D3D12Submission.cpp
@@ -6,6 +6,7 @@
 #include "HAL/RunnableThread.h"
 #include "IRenderCaptureProvider.h"
 #include "Stats/ThreadIdleStats.h"
+#include "ProfilingDebugging/AssetMetadataTrace.h"
 
 #ifndef D3D12_PLATFORM_SUPPORTS_BLOCKING_FENCES
 #define D3D12_PLATFORM_SUPPORTS_BLOCKING_FENCES 1
@@ -964,6 +965,10 @@ void FD3D12PayloadBase::AddQueueFenceWait(FD3D12Fence& InFence, uint64 InValue)
 	QueueFencesToWait.Add({ InFence, InValue });
 }
 
+UE_TRACE_EVENT_BEGIN(Cpu, SyncPoint_Wait, NoSync)
+	UE_TRACE_EVENT_FIELD(UE::Trace::WideString, WaitingForName )
+UE_TRACE_EVENT_END()
+
 void FD3D12SyncPoint::Wait() const
 {
 	checkf(GraphEvent, TEXT("This sync point was not created with a CPU event. Cannot wait for completion on the CPU."));
@@ -971,7 +976,11 @@ void FD3D12SyncPoint::Wait() const
 	if (!GraphEvent->IsComplete())
 	{
 		// Block the calling thread until the graph event is signaled by the interrupt thread.
-		SCOPED_NAMED_EVENT_TEXT("SyncPoint_Wait", FColor::Turquoise);
+#if CPUPROFILERTRACE_ENABLED
+		UE_TRACE_LOG_SCOPED_T(Cpu, SyncPoint_Wait, CpuChannel)
+			<< SyncPoint_Wait.WaitingForName(GetDebugName());
+#endif // CPUPROFILERTRACE_ENABLED
+
 		FD3D12DynamicRHI::GetD3DRHI()->ProcessInterruptQueueUntil(GraphEvent);
 	}
 
diff --git a/Engine/Source/Runtime/D3D12RHI/Private/D3D12Submission.h b/Engine/Source/Runtime/D3D12RHI/Private/D3D12Submission.h
--- a/Engine/Source/Runtime/D3D12RHI/Private/D3D12Submission.h
+++ b/Engine/Source/Runtime/D3D12RHI/Private/D3D12Submission.h
@@ -86,7 +86,10 @@ class FD3D12SyncPoint final : public FThreadSafeRefCountedObject
 	TOptional<FD3D12ResolvedFence> ResolvedFence;
 	FGraphEventRef GraphEvent;
 
-	FD3D12SyncPoint(ED3D12SyncPointType Type)
+	FD3D12SyncPoint(ED3D12SyncPointType Type, const TCHAR* InDebugName)
+#if RHI_USE_SYNC_POINT_DEBUG_NAME
+		:DebugName(InDebugName)
+#endif
 	{
 		if (Type == ED3D12SyncPointType::GPUAndCPU)
 		{
@@ -95,10 +98,10 @@ class FD3D12SyncPoint final : public FThreadSafeRefCountedObject
 	}
 
 public:
-	static FD3D12SyncPointRef Create(ED3D12SyncPointType Type)
+	static FD3D12SyncPointRef Create(ED3D12SyncPointType Type, const TCHAR* InDebugName)
 	{
 		LLM_SCOPE_BYNAME(TEXT("RHIMisc/CreateSyncPoint"));
-		return new FD3D12SyncPoint(Type);
+		return new FD3D12SyncPoint(Type, InDebugName);
 	}
 
 	bool IsComplete() const
@@ -138,6 +141,22 @@ public:
 	{
 		MemoryPool.Push(Pointer);
 	}
+	 
+	inline const TCHAR* GetDebugName() const
+	{
+#if RHI_USE_SYNC_POINT_DEBUG_NAME
+		return DebugName;
+#else
+		return TEXT("SyncPointBusyWait");
+#endif
+	}
+
+private:
+
+#if RHI_USE_SYNC_POINT_DEBUG_NAME
+	const TCHAR* DebugName = nullptr;
+#endif
+
 };
 
 struct FD3D12CommitReservedResourceDesc
diff --git a/Engine/Source/Runtime/D3D12RHI/Private/D3D12Viewport.cpp b/Engine/Source/Runtime/D3D12RHI/Private/D3D12Viewport.cpp
--- a/Engine/Source/Runtime/D3D12RHI/Private/D3D12Viewport.cpp
+++ b/Engine/Source/Runtime/D3D12RHI/Private/D3D12Viewport.cpp
@@ -652,7 +652,7 @@ void FD3D12Viewport::IssueFrameEvent()
 	{
 		FD3D12CommandContext& Context = Device->GetDefaultCommandContext();
 
-		FD3D12SyncPointRef SyncPoint = FD3D12SyncPoint::Create(ED3D12SyncPointType::GPUAndCPU);
+		FD3D12SyncPointRef SyncPoint = FD3D12SyncPoint::Create(ED3D12SyncPointType::GPUAndCPU, TEXT("WaitPreviousFrameEvent"));
 
 		Context.SignalSyncPoint(SyncPoint);
 		Context.Finalize(Payloads);
diff --git a/Engine/Source/Runtime/RHI/Public/RHIDefinitions.h b/Engine/Source/Runtime/RHI/Public/RHIDefinitions.h
--- a/Engine/Source/Runtime/RHI/Public/RHIDefinitions.h
+++ b/Engine/Source/Runtime/RHI/Public/RHIDefinitions.h
@@ -54,6 +54,13 @@ static_assert(sizeof(void*) <= SHADER_PARAMETER_POINTER_ALIGNMENT, "The alignmen
 #endif
 #endif
 
+#ifndef RHI_USE_SYNC_POINT_DEBUG_NAME
+#if (UE_BUILD_SHIPPING)
+#define RHI_USE_SYNC_POINT_DEBUG_NAME	0
+#else
+#define RHI_USE_SYNC_POINT_DEBUG_NAME	1
+#endif
+#endif
 
 #ifndef RHI_RAYTRACING
 #define RHI_RAYTRACING 0
diff --git a/Engine/Source/Runtime/RHI/Public/RHIStrings.h b/Engine/Source/Runtime/RHI/Public/RHIStrings.h
--- a/Engine/Source/Runtime/RHI/Public/RHIStrings.h
+++ b/Engine/Source/Runtime/RHI/Public/RHIStrings.h
@@ -106,3 +106,14 @@ inline const TCHAR* GetShaderFrequencyString(EShaderFrequency Frequency, bool bI
 	String += Index;
 	return String;
 }
+
+inline const TCHAR* GetRenderQueryTypeName(ERenderQueryType QueryType)
+{
+	switch (QueryType)
+	{
+	default: checkNoEntry(); // fallthrough
+	case RQT_Undefined:		return TEXT("Undefined");
+	case RQT_Occlusion:		return TEXT("GPUOcclusion");
+	case RQT_AbsoluteTime:	return TEXT("GPUAbsoluteTime");
+	}
+}
