Allow unbound constant buffer slots if the shader isn't actually using the respective constant buffer.

Empty constant buffers are culled and will not be bound by the high level renderer. However, if a used constant buffer exists at a higher slot index, the state cache will encounter an unset GPU address while iterating through all the slots. FD3D12DescriptorCache::SetRootConstantBuffers wasn't checking if the slot is used by the shader, so it would trigger an incorrect fatal error in this case.

#rnx

diff --git a/Engine/Source/Runtime/D3D12RHI/Private/D3D12DescriptorCache.cpp b/Engine/Source/Runtime/D3D12RHI/Private/D3D12DescriptorCache.cpp
--- a/Engine/Source/Runtime/D3D12RHI/Private/D3D12DescriptorCache.cpp
+++ b/Engine/Source/Runtime/D3D12RHI/Private/D3D12DescriptorCache.cpp
@@ -631,6 +631,43 @@ void FD3D12DescriptorCache::SetConstantBufferViews(EShaderFrequency ShaderStage,
 #endif // D3D12RHI_USE_CONSTANT_BUFFER_VIEWS
 }
 
+static FString GetPipelineShaderHashList(FD3D12GraphicsPipelineState* GraphicsPSO)
+{
+	FString ShaderHashList;
+
+	if (GraphicsPSO == nullptr)
+	{
+		return ShaderHashList;
+	}
+
+	FSHA1 PipelineHasher;
+
+	const auto AddShaderHash = [&PipelineHasher, &ShaderHashList](const FRHIShader* Shader)
+		{
+			FSHAHash ShaderHash;
+			if (Shader)
+			{
+				ShaderHash = Shader->GetHash();
+				ShaderHashList.Appendf(TEXT("%s: %s, "), GetShaderFrequencyString(Shader->GetFrequency(), false), *ShaderHash.ToString());
+			}
+			PipelineHasher.Update(&ShaderHash.Hash[0], sizeof(FSHAHash));
+		};
+
+	AddShaderHash(GraphicsPSO->GetVertexShader());
+	AddShaderHash(GraphicsPSO->GetMeshShader());
+	AddShaderHash(GraphicsPSO->GetAmplificationShader());
+	AddShaderHash(GraphicsPSO->GetPixelShader());
+	AddShaderHash(GraphicsPSO->GetGeometryShader());
+
+	PipelineHasher.Final();
+	FSHAHash PipelineHash;
+	PipelineHasher.GetHash(&PipelineHash.Hash[0]);
+
+	ShaderHashList.Appendf(TEXT("Pipeline: %s"), *PipelineHash.ToString());
+
+	return ShaderHashList;
+}
+
 void FD3D12DescriptorCache::SetRootConstantBuffers(EShaderFrequency ShaderStage, const FD3D12RootSignature* RootSignature, FD3D12ConstantBufferCache& Cache, CBVSlotMask SlotsNeededMask, FD3D12StateCache* StateCache)
 {
 	CBVSlotMask& CurrentDirtySlotMask = Cache.DirtySlotMask[ShaderStage];
@@ -650,104 +687,82 @@ void FD3D12DescriptorCache::SetRootConstantBuffers(EShaderFrequency ShaderStage,
 	for (uint32 SlotIndex = 0; SlotIndex < RDCBVsNeeded; SlotIndex++)
 	{
 		// Only set the root descriptor if it's dirty and we need to set it (it can be used by the shader).
-		if (FD3D12ConstantBufferCache::IsSlotDirty(RDCBVSlotsNeededMask, SlotIndex))
+		if (!FD3D12ConstantBufferCache::IsSlotDirty(RDCBVSlotsNeededMask, SlotIndex))
 		{
-			const D3D12_GPU_VIRTUAL_ADDRESS CurrentGPUVirtualAddress = Cache.CurrentGPUVirtualAddress[ShaderStage][SlotIndex];
-			if (CurrentGPUVirtualAddress == 0)
-			{
-				FString ShaderHashList;
+			continue;
+		}
 
-				switch(ShaderStage)
-				{
-					case SF_Vertex:
-					case SF_Mesh:
-					case SF_Amplification:
-					case SF_Pixel:
-					case SF_Geometry:
-					{
-						FD3D12GraphicsPipelineState* GraphicsPSO = StateCache->GetGraphicsPipelineState();
-						if (!GraphicsPSO)
-						{
-							// Shouldn't happen, but we don't want to crash while crashing.
-							ShaderHashList = TEXT("NO GRAPHICS PSO!");
-							break;
-						}
-
-						FSHA1 PipelineHasher;
-
-						const auto AddShaderHash = [&PipelineHasher, &ShaderHashList](const FRHIShader* Shader)
-						{
-							FSHAHash ShaderHash;
-							if (Shader)
-							{
-								ShaderHash = Shader->GetHash();
-								ShaderHashList.Appendf(TEXT("%s: %s, "), GetShaderFrequencyString(Shader->GetFrequency(), false), *ShaderHash.ToString());
-							}
-							PipelineHasher.Update(&ShaderHash.Hash[0], sizeof(FSHAHash));
-						};
-
-						AddShaderHash(GraphicsPSO->GetVertexShader());
-						AddShaderHash(GraphicsPSO->GetMeshShader());
-						AddShaderHash(GraphicsPSO->GetAmplificationShader());
-						AddShaderHash(GraphicsPSO->GetPixelShader());
-						AddShaderHash(GraphicsPSO->GetGeometryShader());
-
-						PipelineHasher.Final();
-						FSHAHash PipelineHash;
-						PipelineHasher.GetHash(&PipelineHash.Hash[0]);
-
-						ShaderHashList.Appendf(TEXT("Pipeline: %s"), *PipelineHash.ToString());
-						break;
-					}
-
-					case SF_Compute:
-					{
-						FD3D12ComputePipelineState* ComputePSO = StateCache->GetComputePipelineState();
-						if (ComputePSO && ComputePSO->GetComputeShader())
-						{
-							ShaderHashList.Appendf(TEXT("Compute: %s"), *ComputePSO->GetComputeShader()->GetHash().ToString());
-						}
-						else
-						{
-							// Shouldn't happen, but we don't want to crash while crashing.
-							ShaderHashList = TEXT("NO COMPUTE SHADER!");
-						}
-						break;
-					}
-
-					default:
-					{
-						ShaderHashList = TEXT("NO PSO FOR STAGE!");
-						break;
-					}
-				}
+		const D3D12_GPU_VIRTUAL_ADDRESS CurrentGPUVirtualAddress = Cache.CurrentGPUVirtualAddress[ShaderStage][SlotIndex];
 
-				UE_LOG(LogD3D12RHI, Fatal, TEXT("Missing uniform buffer at slot %u, stage %s. Please check the high level drawing code. Hashes: %s."), SlotIndex, GetShaderFrequencyString(ShaderStage), *ShaderHashList);
+		// If the slot is not set in the cache, we need to check if the shader really uses the constant buffer. If the expected layout hash is 0, it means that the
+		// constant buffer is unused, so it's fine if it's not bound. If the hash is non-zero, the shader expects a buffer, so this is a fatal error.
+		if (CurrentGPUVirtualAddress == 0)
+		{
+			FRHIShader* Shader;
+
+			if (ShaderStage == SF_Compute)
+			{
+				FD3D12ComputePipelineState* ComputePSO = StateCache->GetComputePipelineState();
+				Shader = ComputePSO ? ComputePSO->GetComputeShader() : nullptr;
+			}
+			else
+			{
+				FD3D12GraphicsPipelineState* GraphicsPSO = StateCache->GetGraphicsPipelineState();
+				Shader = GraphicsPSO ? GraphicsPSO->GetShader(ShaderStage) : nullptr;
 			}
 
+			// The shader should never be null if we got here, but check anyway because we don't want to crash while crashing.
+			if (Shader == nullptr)
+			{
+				UE_LOG(LogD3D12RHI, Fatal, TEXT("Missing uniform buffer at slot %u, stage %s and no shader is bound."), SlotIndex, GetShaderFrequencyString(ShaderStage));
+				// The above call terminates the process, but static analysis tools don't know that so let's make this easy for them.
+				return;
+			}
 
-			if ((BaseIndex + SlotIndex) < 64)
+			const TArray<uint32>& LayoutHashes = Shader->GetShaderResourceTable().ResourceTableLayoutHashes;
+			uint32 ExpectedLayoutHash = LayoutHashes.IsValidIndex(SlotIndex) ? LayoutHashes[SlotIndex] : 0;
+			if (ExpectedLayoutHash != 0)
 			{
+				FString ShaderHashList;
+
 				if (ShaderStage == SF_Compute)
 				{
-					Context.GraphicsCommandList()->SetComputeRootConstantBufferView(BaseIndex + SlotIndex, CurrentGPUVirtualAddress);
+					ShaderHashList.Appendf(TEXT("Compute: %s"), *Shader->GetHash().ToString());
 				}
 				else
 				{
-					Context.GraphicsCommandList()->SetGraphicsRootConstantBufferView(BaseIndex + SlotIndex, CurrentGPUVirtualAddress);
+					ShaderHashList = GetPipelineShaderHashList(StateCache->GetGraphicsPipelineState());
 				}
+
+				UE_LOG(LogD3D12RHI, Fatal, TEXT("Missing uniform buffer at slot %u, stage %s. Please check the high level drawing code. Hashes: %s."), SlotIndex, GetShaderFrequencyString(ShaderStage), *ShaderHashList);
+				return;
+			}
+		}
+
+		if ((BaseIndex + SlotIndex) < 64)
+		{
+			if (ShaderStage == SF_Compute)
+			{
+				Context.GraphicsCommandList()->SetComputeRootConstantBufferView(BaseIndex + SlotIndex, CurrentGPUVirtualAddress);
 			}
 			else
 			{
-				UE_LOG(LogD3D12RHI, Error, TEXT("%s call exceed max 64 slots: %d "), (ShaderStage == SF_Compute) ? TEXT("SetComputeRootConstantBufferView") : TEXT("SetGraphicsRootConstantBufferView"), BaseIndex + SlotIndex);
+				Context.GraphicsCommandList()->SetGraphicsRootConstantBufferView(BaseIndex + SlotIndex, CurrentGPUVirtualAddress);
 			}
+		}
+		else
+		{
+			UE_LOG(LogD3D12RHI, Error, TEXT("%s call exceed max 64 slots: %d "), (ShaderStage == SF_Compute) ? TEXT("SetComputeRootConstantBufferView") : TEXT("SetGraphicsRootConstantBufferView"), BaseIndex + SlotIndex);
+		}
 
-			// Update residency.
+		// Update residency is the slot is used.
+		if (CurrentGPUVirtualAddress != 0)
+		{
 			Context.UpdateResidency(Cache.Resources[ShaderStage][SlotIndex]);
-
-			// Clear the dirty bit.
-			FD3D12ConstantBufferCache::CleanSlot(CurrentDirtySlotMask, SlotIndex);
 		}
+
+		// Clear the dirty bit.
+		FD3D12ConstantBufferCache::CleanSlot(CurrentDirtySlotMask, SlotIndex);
 	}
 	check((CurrentDirtySlotMask & RDCBVSlotsNeededMask) == 0);	// Check all slots that needed to be set, were set.
 
