Do not calibrate GPU timestamps for D3D12 copy queues if not supported.

diff --git a/Engine/Source/Runtime/D3D12RHI/Private/D3D12Submission.cpp b/Engine/Source/Runtime/D3D12RHI/Private/D3D12Submission.cpp
--- a/Engine/Source/Runtime/D3D12RHI/Private/D3D12Submission.cpp
+++ b/Engine/Source/Runtime/D3D12RHI/Private/D3D12Submission.cpp
@@ -902,12 +902,14 @@ void FD3D12DynamicRHI::FlushBatchedPayloads(FD3D12Queue::FPayloadArray& Payloads
 
 			if (FD3D12Timing* LocalTiming = *Payload->Timing)
 			{
-				SCOPED_NAMED_EVENT(CalibrateClocks, FColor::Red);
-
-				// Calibrate the GPU timestamp / clock
-				VERIFYD3D12RESULT(Payload->Queue.D3DCommandQueue->GetClockCalibration(&LocalTiming->GPUTimestamp, &LocalTiming->CPUTimestamp));
-				VERIFYD3D12RESULT(Payload->Queue.D3DCommandQueue->GetTimestampFrequency(&LocalTiming->GPUFrequency));
-				QueryPerformanceFrequency(reinterpret_cast<LARGE_INTEGER*>(&LocalTiming->CPUFrequency));
+				// Calibrate the GPU timestamp / clock, if the queue type supports calibration.
+				if (Payload->Queue.QueueType != ED3D12QueueType::Copy || Payload->Queue.Device->GetParentAdapter()->AreCopyQueueTimestampQueriesSupported())
+				{
+					SCOPED_NAMED_EVENT(CalibrateClocks, FColor::Red);
+					VERIFYD3D12RESULT(Payload->Queue.D3DCommandQueue->GetClockCalibration(&LocalTiming->GPUTimestamp, &LocalTiming->CPUTimestamp));
+					VERIFYD3D12RESULT(Payload->Queue.D3DCommandQueue->GetTimestampFrequency(&LocalTiming->GPUFrequency));
+					QueryPerformanceFrequency(reinterpret_cast<LARGE_INTEGER*>(&LocalTiming->CPUFrequency));
+				}
 			}
 		}
 #endif // RHI_NEW_GPU_PROFILER
