Log out D3D10Warp dll version analogous to the driver version when using Warp.

#rnx

diff --git a/Engine/Source/Runtime/D3D12RHI/Private/Windows/WindowsD3D12Device.cpp b/Engine/Source/Runtime/D3D12RHI/Private/Windows/WindowsD3D12Device.cpp
--- a/Engine/Source/Runtime/D3D12RHI/Private/Windows/WindowsD3D12Device.cpp
+++ b/Engine/Source/Runtime/D3D12RHI/Private/Windows/WindowsD3D12Device.cpp
@@ -917,6 +917,20 @@ void FD3D12DynamicRHIModule::FindAdapter()
 					UE_LOG(LogD3D12RHI, Log, TEXT("  Driver Version: %s (internal:%s, unified:%s)"), *GPUDriverInfo.UserDriverVersion, *GPUDriverInfo.InternalDriverVersion, *GPUDriverInfo.GetUnifiedDriverVersion());
 					UE_LOG(LogD3D12RHI, Log, TEXT("     Driver Date: %s"), *GPUDriverInfo.DriverDate);
 				}
+				else
+				{
+					const TCHAR* WarpDllName = TEXT("d3d10warp.dll");
+					void* WarpDllHandle = FPlatformProcess::GetDllHandle(WarpDllName);
+					if(WarpDllHandle)
+					{
+						const uint64 WarpVersion = FPlatformMisc::GetFileVersion(WarpDllName);
+						const uint64 Major = ((WarpVersion >> 48) & 0xffff);
+						const uint64 Minor = ((WarpVersion >> 32) & 0xffff);
+						const uint64 Build = ((WarpVersion >> 16) & 0xffff);
+						const uint64 Revision = (WarpVersion & 0xffff);
+						UE_LOG(LogD3D12RHI, Log, TEXT("  D3D10Warp Version: %d.%d.%d.%d"), Major, Minor, Build, Revision);
+					}
+				}
 
 				FD3D12AdapterDesc CurrentAdapter(AdapterDesc, AdapterIndex, DeviceInfo);
 
