D3D12 Bindless - fixing issues when switching back from explicit descriptor heaps.

diff --git a/Engine/Source/Runtime/D3D12RHI/Private/D3D12DescriptorCache.cpp b/Engine/Source/Runtime/D3D12RHI/Private/D3D12DescriptorCache.cpp
--- a/Engine/Source/Runtime/D3D12RHI/Private/D3D12DescriptorCache.cpp
+++ b/Engine/Source/Runtime/D3D12RHI/Private/D3D12DescriptorCache.cpp
@@ -869,6 +869,7 @@ void FD3D12DescriptorCache::SetExplicitDescriptorCache(FD3D12ExplicitDescriptorC
 		bExplicitViewHeapIsBindless = bViewHeapIsBindless;
 
 #if PLATFORM_SUPPORTS_BINDLESS_RENDERING
+		bLastSetHeapBeforeExplicitIsBindless = bLastSetHeapsBindless;
 		bLastSetHeapsBindless = bViewHeapIsBindless;
 #endif
 	}
@@ -878,7 +879,14 @@ void FD3D12DescriptorCache::UnsetExplicitDescriptorCache()
 {
 	if (bUsingExplicitCacheHeaps)
 	{
-		SetDescriptorHeaps(ED3D12SetDescriptorHeapsFlags::None);
+		ED3D12SetDescriptorHeapsFlags Flags = ED3D12SetDescriptorHeapsFlags::None;
+#if PLATFORM_SUPPORTS_BINDLESS_RENDERING
+		if (bLastSetHeapBeforeExplicitIsBindless)
+		{
+			EnumAddFlags(Flags, ED3D12SetDescriptorHeapsFlags::Bindless);
+		}
+#endif
+		SetDescriptorHeaps(Flags);
 
 		bUsingExplicitCacheHeaps = false;
 		bExplicitViewHeapIsBindless = false;
diff --git a/Engine/Source/Runtime/D3D12RHI/Private/D3D12DescriptorCache.h b/Engine/Source/Runtime/D3D12RHI/Private/D3D12DescriptorCache.h
--- a/Engine/Source/Runtime/D3D12RHI/Private/D3D12DescriptorCache.h
+++ b/Engine/Source/Runtime/D3D12RHI/Private/D3D12DescriptorCache.h
@@ -441,6 +441,7 @@ private:
 	bool bFullyBindless = false;
 
 	bool bLastSetHeapsBindless = false;
+	bool bLastSetHeapBeforeExplicitIsBindless = false;
 
 	FD3D12DescriptorHeapPtr BindlessResourcesHeap = nullptr;
 	FD3D12DescriptorHeapPtr BindlessSamplersHeap = nullptr;
