Defer deletion of ID3D12PipelineState objects.

diff --git a/Engine/Source/Runtime/D3D12RHI/Private/D3D12Pipelinestate.cpp b/Engine/Source/Runtime/D3D12RHI/Private/D3D12Pipelinestate.cpp
--- a/Engine/Source/Runtime/D3D12RHI/Private/D3D12Pipelinestate.cpp
+++ b/Engine/Source/Runtime/D3D12RHI/Private/D3D12Pipelinestate.cpp
@@ -314,6 +314,12 @@ FD3D12PipelineState::~FD3D12PipelineState()
 		Worker = nullptr;
 	}
 
+	if (UsePSORefCounting())
+	{
+		PipelineState->AddRef(); // AddRef for DeferredDelete to take over the actual deletion
+		FD3D12DynamicRHI::GetD3DRHI()->DeferredDelete(PipelineState);
+	}
+
 	DEC_DWORD_STAT(STAT_D3D12NumPSOs);
 }
 
