Enable extra check on D3D12 cmd list when ENABLE_RHI_VALIDATION is supported, those errors are not well reported from the debug layer

diff --git a/Engine/Source/Runtime/D3D12RHI/Private/D3D12CommandList.cpp b/Engine/Source/Runtime/D3D12RHI/Private/D3D12CommandList.cpp
--- a/Engine/Source/Runtime/D3D12RHI/Private/D3D12CommandList.cpp
+++ b/Engine/Source/Runtime/D3D12RHI/Private/D3D12CommandList.cpp
@@ -461,13 +461,15 @@ void FD3D12ContextCommon::TransitionResource(FD3D12Resource* Resource, D3D12_RES
 #ifdef PLATFORM_SUPPORTS_RESOURCE_COMPRESSION
 	After |= Resource->GetCompressedState();
 #endif
-#if DEBUG_RESOURCE_STATES
+
+#if ENABLE_RHI_VALIDATION
 	FString IncompatibilityReason;
 	if (CheckResourceStateCompatibility(After, Resource->GetDesc().Flags, IncompatibilityReason) == false)
 	{
 		UE_LOG(LogRHI, Error, TEXT("Incompatible Transition State for Resource %s - %s"), *Resource->GetName().ToString(), *IncompatibilityReason);
 	}
 #endif
+
 	UpdateResidency(Resource);
 
 	TransitionResource(Resource, Subresource, Before, After);
