Remove check for a null resource pointer since it might be a valid case and the system will be able to recover from it.

diff --git a/Engine/Source/Runtime/D3D12RHI/Private/D3D12Commands.cpp b/Engine/Source/Runtime/D3D12RHI/Private/D3D12Commands.cpp
--- a/Engine/Source/Runtime/D3D12RHI/Private/D3D12Commands.cpp
+++ b/Engine/Source/Runtime/D3D12RHI/Private/D3D12Commands.cpp
@@ -246,7 +246,7 @@ std::tuple<FD3D12Resource*,FD3D12Texture*> GetResourceAndTexture(FD3D12CommandCo
 	{
 		// Resource may be null if this is a multi-GPU resource not present on the current GPU
 		FD3D12Buffer* Buffer = Context.RetrieveObject<FD3D12Buffer>(Info.Buffer);
-		check(Buffer || GNumExplicitGPUsForRendering > 1);
+
 		if (Buffer)
 		{
 			return { Buffer->GetResource(), nullptr };
@@ -257,7 +257,7 @@ std::tuple<FD3D12Resource*,FD3D12Texture*> GetResourceAndTexture(FD3D12CommandCo
 	{
 		// Resource may be null if this is a multi-GPU resource not present on the current GPU
 		FD3D12Texture* Texture = Context.RetrieveTexture(Info.Texture);
-		check(Texture || GNumExplicitGPUsForRendering > 1);
+
 		if (Texture)
 		{
 			return { Texture->GetResource(), Texture };
diff --git a/Engine/Source/Runtime/D3D12RHI/Private/D3D12Texture.h b/Engine/Source/Runtime/D3D12RHI/Private/D3D12Texture.h
--- a/Engine/Source/Runtime/D3D12RHI/Private/D3D12Texture.h
+++ b/Engine/Source/Runtime/D3D12RHI/Private/D3D12Texture.h
@@ -241,7 +241,7 @@ FORCEINLINE FD3D12Texture* GetD3D12TextureFromRHITexture(FRHITexture* Texture)
 	
 	FRHITexture* RHITexture = Texture;
 	FD3D12Texture* Result((FD3D12Texture*)RHITexture->GetTextureBaseRHI());
-	check(Result);
+	
 	return Result;
 }
 
