Changed number of DX12 viewport buffers to be driven by a read-only CVar rather than macro.

diff --git a/Engine/Source/Runtime/D3D12RHI/Private/D3D12Viewport.cpp b/Engine/Source/Runtime/D3D12RHI/Private/D3D12Viewport.cpp
--- a/Engine/Source/Runtime/D3D12RHI/Private/D3D12Viewport.cpp
+++ b/Engine/Source/Runtime/D3D12RHI/Private/D3D12Viewport.cpp
@@ -33,6 +33,13 @@ namespace D3D12RHI
 		);
 #endif
 	};
+
+	int32 GD3D12NumViewportBuffers = D3D12RHI_DEFAULT_NUM_BACKBUFFER;
+	static FAutoConsoleVariableRef CVarD3D12NumViewportBuffers(
+		TEXT("r.D3D12.NumViewportBuffers"),
+		GD3D12NumViewportBuffers,
+		TEXT("The number of viewport buffers to allocate for the RHI viewport."),
+		ECVF_ReadOnly);
 }
 using namespace D3D12RHI;
 
@@ -686,7 +693,9 @@ FViewportRHIRef FD3D12DynamicRHI::RHICreateViewport(void* WindowHandle, uint32 S
 		PreferredPixelFormat = GetDefaultBackBufferPixelFormat();
 	}
 
-	FD3D12Viewport* RenderingViewport = new FD3D12Viewport(&GetAdapter(), (HWND)WindowHandle, SizeX, SizeY, bIsFullscreen, PreferredPixelFormat);
+	const uint32 InNumBackBuffers = GD3D12NumViewportBuffers;
+
+	FD3D12Viewport* RenderingViewport = new FD3D12Viewport(&GetAdapter(), (HWND)WindowHandle, SizeX, SizeY, bIsFullscreen, PreferredPixelFormat, InNumBackBuffers);
 	RenderingViewport->Init();
 	return RenderingViewport;
 }
diff --git a/Engine/Source/Runtime/D3D12RHI/Private/D3D12Viewport.h b/Engine/Source/Runtime/D3D12RHI/Private/D3D12Viewport.h
--- a/Engine/Source/Runtime/D3D12RHI/Private/D3D12Viewport.h
+++ b/Engine/Source/Runtime/D3D12RHI/Private/D3D12Viewport.h
@@ -29,7 +29,7 @@ public:
 	// Reported & known problem in DXGI and will be fixed with future release but DXGI is not part of the Agility SDK so a code side fix is needed for now.
 	static FCriticalSection DXGIBackBufferLock;
 
-	FD3D12Viewport(class FD3D12Adapter* InParent, HWND InWindowHandle, uint32 InSizeX, uint32 InSizeY, bool bInIsFullscreen, EPixelFormat InPixelFormat);
+	FD3D12Viewport(class FD3D12Adapter* InParent, HWND InWindowHandle, uint32 InSizeX, uint32 InSizeY, bool bInIsFullscreen, EPixelFormat InPixelFormat, uint32 InNumBackBuffers);
 
 	void Init();
 
@@ -249,8 +249,8 @@ private:
 #endif
 	};
 
-	static constexpr uint32 NumBackBuffers = GD3D12RHINumBackBuffers;
-	TStaticArray<FBackBufferData, NumBackBuffers> BackBuffers;
+	uint32 NumBackBuffers = 0;
+	TArray<FBackBufferData, TInlineAllocator<D3D12RHI_DEFAULT_NUM_BACKBUFFER>> BackBuffers;
 
 #if D3D12RHI_USE_DUMMY_BACKBUFFER
 	// Dummy back buffer texture which always references the current back buffer on the RHI thread
diff --git a/Engine/Source/Runtime/D3D12RHI/Private/Windows/WindowsD3D12RHIDefinitions.h b/Engine/Source/Runtime/D3D12RHI/Private/Windows/WindowsD3D12RHIDefinitions.h
--- a/Engine/Source/Runtime/D3D12RHI/Private/Windows/WindowsD3D12RHIDefinitions.h
+++ b/Engine/Source/Runtime/D3D12RHI/Private/Windows/WindowsD3D12RHIDefinitions.h
@@ -4,7 +4,7 @@
 
 #include "D3D12ThirdParty.h"
 
-constexpr uint32 GD3D12RHINumBackBuffers = 3;
+#define D3D12RHI_DEFAULT_NUM_BACKBUFFER 3
 
 // Windows uses the graphics command list interface for the copy queue.
 typedef ID3D12GraphicsCommandList ID3D12CopyCommandList;
diff --git a/Engine/Source/Runtime/D3D12RHI/Private/Windows/WindowsD3D12Viewport.cpp b/Engine/Source/Runtime/D3D12RHI/Private/Windows/WindowsD3D12Viewport.cpp
--- a/Engine/Source/Runtime/D3D12RHI/Private/Windows/WindowsD3D12Viewport.cpp
+++ b/Engine/Source/Runtime/D3D12RHI/Private/Windows/WindowsD3D12Viewport.cpp
@@ -23,7 +23,7 @@ static FAutoConsoleVariableRef CVarD3DUseAllowTearing(
 	ECVF_RenderThreadSafe | ECVF_ReadOnly
 );
 
-FD3D12Viewport::FD3D12Viewport(class FD3D12Adapter* InParent, HWND InWindowHandle, uint32 InSizeX, uint32 InSizeY, bool bInIsFullscreen, EPixelFormat InPreferredPixelFormat)
+FD3D12Viewport::FD3D12Viewport(class FD3D12Adapter* InParent, HWND InWindowHandle, uint32 InSizeX, uint32 InSizeY, bool bInIsFullscreen, EPixelFormat InPreferredPixelFormat, uint32 InNumBackBuffers)
 	: FD3D12AdapterChild(InParent)
 	, WindowHandle(InWindowHandle)
 	, SizeX(InSizeX)
@@ -31,7 +31,10 @@ FD3D12Viewport::FD3D12Viewport(class FD3D12Adapter* InParent, HWND InWindowHandl
 	, PixelFormat(InPreferredPixelFormat)
 	, bIsFullscreen(bInIsFullscreen)
 	, bNeedSwapChain(!FParse::Param(FCommandLine::Get(), TEXT("RenderOffScreen")))
+	, NumBackBuffers(InNumBackBuffers)
 {
+	check(NumBackBuffers > 0);
+	BackBuffers.SetNum(NumBackBuffers);
 	check(IsInGameThread());
 	GetParentAdapter()->GetViewports().Add(this);
 }
