PR #13437: Fix(D3D12): Correct operator== logic for FD3D12ResourceInitConfig

diff --git a/Engine/Source/Runtime/D3D12RHI/Private/D3D12PoolAllocator.h b/Engine/Source/Runtime/D3D12RHI/Private/D3D12PoolAllocator.h
--- a/Engine/Source/Runtime/D3D12RHI/Private/D3D12PoolAllocator.h
+++ b/Engine/Source/Runtime/D3D12RHI/Private/D3D12PoolAllocator.h
@@ -35,7 +35,7 @@ struct FD3D12ResourceInitConfig
 	bool operator==(const FD3D12ResourceInitConfig& InOther) const
 	{
 		return HeapType == InOther.HeapType &&
-			HeapFlags && InOther.HeapFlags &&
+			HeapFlags == InOther.HeapFlags &&
 			ResourceFlags == InOther.ResourceFlags &&
 			InitialResourceState == InOther.InitialResourceState;
 	}
