Fix GetResourceInfo for reserved resources on d3d12 for rhi memory dumps

diff --git a/Engine/Source/Runtime/D3D12RHI/Private/D3D12Resources.h b/Engine/Source/Runtime/D3D12RHI/Private/D3D12Resources.h
--- a/Engine/Source/Runtime/D3D12RHI/Private/D3D12Resources.h
+++ b/Engine/Source/Runtime/D3D12RHI/Private/D3D12Resources.h
@@ -366,6 +366,12 @@ public:
 	inline FD3D12Heap* GetHeap() const { return Heap; };
 	inline bool IsDepthStencilResource() const { return bDepthStencil; }
 
+	inline uint64 GetCommittedReservedResourceSize() const
+	{
+		check(IsReservedResource());
+		return ReservedResourceData->NumCommittedTiles * (uint64)GRHIGlobals.ReservedResources.TileSizeInBytes;
+	}
+
 	inline bool NeedsDeferredResidencyUpdate() const { return IsReservedResource(); }
 
 	void StartTrackingForResidency();
@@ -1037,7 +1043,14 @@ public:
 		OutResourceInfo = FRHIResourceInfo{};
 		OutResourceInfo.Name = GetName();
 		OutResourceInfo.Type = GetType();
-		OutResourceInfo.VRamAllocation.AllocationSize = ResourceLocation.GetSize();
+		if (ResourceLocation.GetResource() && ResourceLocation.GetResource()->IsReservedResource())
+		{
+			OutResourceInfo.VRamAllocation.AllocationSize = ResourceLocation.GetResource()->GetCommittedReservedResourceSize();
+		}
+		else
+		{
+			OutResourceInfo.VRamAllocation.AllocationSize = ResourceLocation.GetSize();
+		}
 		OutResourceInfo.IsTransient = ResourceLocation.IsTransient();
 		OutResourceInfo.bResident = GetResource() && GetResource()->IsResident();
 
