Fixed late shader association printing

diff --git a/Engine/Source/Runtime/D3D12RHI/Private/D3D12NvidiaExtensions.cpp b/Engine/Source/Runtime/D3D12RHI/Private/D3D12NvidiaExtensions.cpp
--- a/Engine/Source/Runtime/D3D12RHI/Private/D3D12NvidiaExtensions.cpp
+++ b/Engine/Source/Runtime/D3D12RHI/Private/D3D12NvidiaExtensions.cpp
@@ -179,7 +179,7 @@ namespace UE::RHICore::Nvidia::Aftermath::D3D12
 			{
 				UE_CALL_ONCE([Elapsed]
 				{
-					UE_LOG(LogD3D12RHI, Warning, TEXT("Late shader associations timed out at %f.5ms"), Elapsed);
+					UE_LOG(LogD3D12RHI, Warning, TEXT("Late shader associations timed out at %.0fms"), Elapsed);
 				});
 				return;
 			}
@@ -226,7 +226,7 @@ namespace UE::RHICore::Nvidia::Aftermath::D3D12
 		});
 		
 		double TimeMS = FPlatformTime::ToMilliseconds64(FPlatformTime::Cycles64() - CycleStart);
-		UE_LOG(LogD3D12RHI, Log, TEXT("Created late shader associations, took %f.5ms"), TimeMS);
+		UE_LOG(LogD3D12RHI, Log, TEXT("Created late shader associations, took %.0fms"), TimeMS);
 	}
 
 	FCommandList RegisterCommandList(ID3D12CommandList* D3DCommandList)
diff --git a/Engine/Source/Runtime/VulkanRHI/Private/VulkanUtil.cpp b/Engine/Source/Runtime/VulkanRHI/Private/VulkanUtil.cpp
--- a/Engine/Source/Runtime/VulkanRHI/Private/VulkanUtil.cpp
+++ b/Engine/Source/Runtime/VulkanRHI/Private/VulkanUtil.cpp
@@ -551,7 +551,7 @@ void AftermathLateAssociate(float TimeLimitSeconds, uint32 FrameLimit)
 	});
 	
 	double TimeMS = FPlatformTime::ToMilliseconds64(FPlatformTime::Cycles64() - CycleStart);
-	UE_LOG(LogVulkanRHI, Log, TEXT("Created late shader associations, took %f.5ms"), TimeMS);
+	UE_LOG(LogVulkanRHI, Log, TEXT("Created late shader associations, took %.0fms"), TimeMS);
 }
 #endif // NV_AFTERMATH
 
