Remove code which disables mesh shaders on Intel GPUs.

#rnx

diff --git a/Engine/Source/Runtime/D3D12RHI/Private/Windows/WindowsD3D12Device.cpp b/Engine/Source/Runtime/D3D12RHI/Private/Windows/WindowsD3D12Device.cpp
--- a/Engine/Source/Runtime/D3D12RHI/Private/Windows/WindowsD3D12Device.cpp
+++ b/Engine/Source/Runtime/D3D12RHI/Private/Windows/WindowsD3D12Device.cpp
@@ -1691,12 +1691,6 @@ void FD3D12DynamicRHI::Init()
 		// Create a new Intel Extensions Device Extension Context to support DX12 extension calls
 		INTCExtensionInfo INTCExtensionInfo{};
 		IntelExtensionContext = CreateIntelExtensionsContext(GetAdapter().GetD3DDevice(), INTCExtensionInfo);
-
-		if (GRHISupportsMeshShadersTier0 || GRHISupportsMeshShadersTier1)
-		{
-			GRHISupportsMeshShadersTier0 = GRHISupportsMeshShadersTier1 = false;
-			UE_LOG(LogD3D12RHI, Warning, TEXT("Mesh shaders are disabled due to a known driver issue."));
-		}
 	}
 #endif // INTEL_EXTENSIONS
 
