Allow using NVIDIA driver-aware precache optimization on DX12 SM5.

Tested that the runtime and the driver behave the same way and that only shader bytecode changes (and root signature, which we deterministically base off the shaders) are important.

#rnx

diff --git a/Engine/Source/Runtime/D3D12RHI/Private/D3D12State.cpp b/Engine/Source/Runtime/D3D12RHI/Private/D3D12State.cpp
--- a/Engine/Source/Runtime/D3D12RHI/Private/D3D12State.cpp
+++ b/Engine/Source/Runtime/D3D12RHI/Private/D3D12State.cpp
@@ -508,7 +508,7 @@ uint64 FD3D12DynamicRHI::RHIComputeStatePrecachePSOHash(const FGraphicsPipelineS
 {
 	if (bDriverCacheAwarePSOPrecaching)
 	{
-		if (IsRHIDeviceNVIDIA() && GMaxRHIFeatureLevel == ERHIFeatureLevel::SM6)
+		if (IsRHIDeviceNVIDIA())
 		{
 			// NVIDIA drivers only care about the shaders for PSO caching.
 			struct FHashKey
@@ -590,7 +590,7 @@ uint64 FD3D12DynamicRHI::RHIComputePrecachePSOHash(const FGraphicsPipelineStateI
 
 	if (bDriverCacheAwarePSOPrecaching)
 	{
-		if (IsRHIDeviceNVIDIA() && GMaxRHIFeatureLevel == ERHIFeatureLevel::SM6)
+		if (IsRHIDeviceNVIDIA())
 		{
 			// We already hashed everything we needed.
 			return Initializer.StatePrecachePSOHash;
