D3D12RHI: dump the active breadcrumbs tree when a GPU payload takes too long
- This will help investigate crash reports where the full GPU crash handling was not triggered, but the interrupt thread detected the timeout of a payload.
- Currently only the timeout message is printed to the log. Now we'll also get the active breadcrumbs from each queue.

diff --git a/Engine/Source/Runtime/D3D12RHI/Private/D3D12Device.cpp b/Engine/Source/Runtime/D3D12RHI/Private/D3D12Device.cpp
--- a/Engine/Source/Runtime/D3D12RHI/Private/D3D12Device.cpp
+++ b/Engine/Source/Runtime/D3D12RHI/Private/D3D12Device.cpp
@@ -40,9 +40,10 @@ FD3D12Timing::FD3D12Timing(FD3D12Queue& Queue)
 #endif
 {}
 
-FD3D12Queue::FD3D12Queue(FD3D12Device* Device, ED3D12QueueType QueueType)
+FD3D12Queue::FD3D12Queue(FD3D12Device* Device, ED3D12QueueType QueueType, int32 QueueIndex)
 	: Device(Device)
 	, QueueType(QueueType)
+	, QueueIndex(QueueIndex)
 	, Fence(this)
 	, BarrierTimestamps(Device, QueueType, D3D12_QUERY_TYPE_TIMESTAMP)
 	, bSupportsTileMapping(FD3D12DynamicRHI::GetD3DRHI()->QueueSupportsTileMapping(QueueType))
@@ -129,9 +130,9 @@ FD3D12Device::FD3D12Device(FRHIGPUMask InGPUMask, FD3D12Adapter* InAdapter)
 		OfflineDescriptorManagers.Emplace(this, (ERHIDescriptorHeapType)HeapType);
 	}
 
-	for (uint32 QueueType = 0; QueueType < (uint32)ED3D12QueueType::Count; ++QueueType)
+	for (uint32 QueueIndex = 0; QueueIndex < (uint32)ED3D12QueueType::Count; ++QueueIndex)
 	{
-		Queues.Emplace(this, (ED3D12QueueType)QueueType);
+		Queues.Emplace(this, (ED3D12QueueType)QueueIndex, QueueIndex);
 	}
 
 	// Some hardware is not capable of running tile mapping operations on all queue types.
diff --git a/Engine/Source/Runtime/D3D12RHI/Private/D3D12Device.h b/Engine/Source/Runtime/D3D12RHI/Private/D3D12Device.h
--- a/Engine/Source/Runtime/D3D12RHI/Private/D3D12Device.h
+++ b/Engine/Source/Runtime/D3D12RHI/Private/D3D12Device.h
@@ -19,6 +19,8 @@ D3D12Device.h: D3D12 Device Interfaces
 #include "D3D12GPUProfiler.h"
 
 #include "Containers/LruCache.h"
+#include "Containers/SpscQueue.h"
+#include "Containers/MpscQueue.h"
 
 class FD3D12Device;
 class FD3D12DynamicRHI;
@@ -103,6 +105,7 @@ class FD3D12Queue final
 public:
 	FD3D12Device* const Device;
 	ED3D12QueueType const QueueType;
+	int32 const QueueIndex;
 
 	// The underlying D3D queue object
 	TRefCountPtr<ID3D12CommandQueue> D3DCommandQueue;
@@ -110,16 +113,8 @@ public:
 	// A single D3D fence to manage completion of work on this queue
 	FD3D12Fence Fence;
 
-	struct : public TQueue<FD3D12Payload*, EQueueMode::Mpsc>
-	{
-		FD3D12Payload* Peek()
-		{
-			if (FD3D12Payload** Result = TQueue::Peek())
-				return *Result;
-
-			return nullptr;
-		}		
-	} PendingSubmission, PendingInterrupt;
+	TMpscQueue<FD3D12Payload*> PendingSubmission;
+	TSpscQueue<FD3D12Payload*> PendingInterrupt;
 
 	FD3D12Payload*          PayloadToSubmit  = nullptr;
 	FD3D12CommandAllocator* BarrierAllocator = nullptr;
@@ -162,7 +157,7 @@ public:
 #endif
 	);
 
-	FD3D12Queue(FD3D12Device* Device, ED3D12QueueType QueueType);
+	FD3D12Queue(FD3D12Device* Device, ED3D12QueueType QueueType, int32 QueueIndex);
 	~FD3D12Queue();
 
 #if RHI_NEW_GPU_PROFILER
diff --git a/Engine/Source/Runtime/D3D12RHI/Private/D3D12RHI.cpp b/Engine/Source/Runtime/D3D12RHI/Private/D3D12RHI.cpp
--- a/Engine/Source/Runtime/D3D12RHI/Private/D3D12RHI.cpp
+++ b/Engine/Source/Runtime/D3D12RHI/Private/D3D12RHI.cpp
@@ -939,6 +939,10 @@ void FD3D12DynamicRHI::HandleGpuTimeout(FD3D12Payload* Payload, double SecondsSi
 		, GetD3DCommandQueueTypeName(Payload->Queue.QueueType)
 		, SecondsSinceSubmission
 	);
+
+#if WITH_RHI_BREADCRUMBS
+	DumpActiveBreadcrumbs(FRHIBreadcrumbState::EVerbosity::Warning);
+#endif
 }
 
 void FD3D12DynamicRHI::SetupD3D12Debug()
diff --git a/Engine/Source/Runtime/D3D12RHI/Private/D3D12RHIPrivate.h b/Engine/Source/Runtime/D3D12RHI/Private/D3D12RHIPrivate.h
--- a/Engine/Source/Runtime/D3D12RHI/Private/D3D12RHIPrivate.h
+++ b/Engine/Source/Runtime/D3D12RHI/Private/D3D12RHIPrivate.h
@@ -177,6 +177,12 @@ public:
 	// Main crash function for handling D3D API errors
 	void HandleFailedD3D12Result(HRESULT D3DResult, ID3D12Device* Device, bool bCreatingTextures, const TCHAR* Message);
 
+#if WITH_RHI_BREADCRUMBS
+	// Prints the active breadcrumb tree to the log.
+	// Can only be called on the interrupt thread.
+	void DumpActiveBreadcrumbs(FRHIBreadcrumbState::EVerbosity Verbosity);
+#endif
+
 private:
 	// Calls OutputGpuCrashReport, and also forces a quit.
 	void TerminateOnGPUCrash();
diff --git a/Engine/Source/Runtime/D3D12RHI/Private/D3D12Submission.cpp b/Engine/Source/Runtime/D3D12RHI/Private/D3D12Submission.cpp
--- a/Engine/Source/Runtime/D3D12RHI/Private/D3D12Submission.cpp
+++ b/Engine/Source/Runtime/D3D12RHI/Private/D3D12Submission.cpp
@@ -430,10 +430,12 @@ FD3D12DynamicRHI::FProcessResult FD3D12DynamicRHI::ProcessSubmissionQueue()
 			while (true)
 			{
 				{
-					FD3D12Payload* Payload = CurrentQueue.PendingSubmission.Peek();
-					if (!Payload)
+					FD3D12Payload** PayloadPtr = CurrentQueue.PendingSubmission.Peek();
+					if (!PayloadPtr)
 						return;
 
+					FD3D12Payload* Payload = *PayloadPtr;
+
 					// Accumulate the list of fences to await, and their maximum values
 					while (Payload->SyncPointsToWait.Index < Payload->SyncPointsToWait.Num())
 					{
@@ -460,7 +462,7 @@ FD3D12DynamicRHI::FProcessResult FD3D12DynamicRHI::ProcessSubmissionQueue()
 				
 					// All necessary sync points have been resolved.
 					Payload->SyncPointsToWait = {};
-					CurrentQueue.PendingSubmission.Pop();
+					CurrentQueue.PendingSubmission.Dequeue();
 					bProgress = true;
 
 					check(!CurrentQueue.PayloadToSubmit);
@@ -1138,8 +1140,10 @@ FD3D12DynamicRHI::FProcessResult FD3D12DynamicRHI::ProcessInterruptQueue()
 	FProcessResult Result;
 	ForEachQueue([&](FD3D12Queue& CurrentQueue)
 	{
-		while (FD3D12Payload* Payload = CurrentQueue.PendingInterrupt.Peek())
+		while (FD3D12Payload** PayloadPtr = CurrentQueue.PendingInterrupt.Peek())
 		{
+			FD3D12Payload* Payload = *PayloadPtr;
+
 			if (!Payload->bSubmitted)
 				break;
 
@@ -1207,7 +1211,7 @@ FD3D12DynamicRHI::FProcessResult FD3D12DynamicRHI::ProcessInterruptQueue()
 
 			// At this point, the current command list has completed on the GPU.
 			CurrentQueue.Fence.bInterruptAwaited = false;
-			CurrentQueue.PendingInterrupt.Pop();
+			CurrentQueue.PendingInterrupt.Dequeue();
 			Result.Status |= EQueueStatus::Processed;
 
 			// Resolve query results
diff --git a/Engine/Source/Runtime/D3D12RHI/Private/D3D12Util.cpp b/Engine/Source/Runtime/D3D12RHI/Private/D3D12Util.cpp
--- a/Engine/Source/Runtime/D3D12RHI/Private/D3D12Util.cpp
+++ b/Engine/Source/Runtime/D3D12RHI/Private/D3D12Util.cpp
@@ -819,6 +819,55 @@ void FD3D12DynamicRHI::TerminateOnOutOfMemory(ID3D12Device* InDevice, HRESULT D3
 #endif // !PLATFORM_WINDOWS
 }
 
+#if WITH_RHI_BREADCRUMBS
+void FD3D12DynamicRHI::DumpActiveBreadcrumbs(FRHIBreadcrumbState::EVerbosity Verbosity)
+{
+	check(IsInInterruptThread());
+
+	if (UE::RHI::UseGPUCrashBreadcrumbs())
+	{
+		TMap<FRHIBreadcrumbState::FQueueID, TArray<FRHIBreadcrumbRange>> QueueRanges;
+		FRHIBreadcrumbState BreadcrumbState;
+
+		ForEachQueue([&](FD3D12Queue& Queue)
+		{
+			if (!Queue.DiagnosticBuffer)
+				return;
+
+			uint32 const DeviceIndex = Queue.Device->GetGPUIndex();
+
+			ERHIPipeline Pipeline;
+			switch (Queue.QueueType)
+			{
+			default: return; // Skip pipelines that the RHI doesn't handle
+			case ED3D12QueueType::Direct: Pipeline = ERHIPipeline::Graphics; break;
+			case ED3D12QueueType::Async: Pipeline = ERHIPipeline::AsyncCompute; break;
+			}
+
+			TArray<FRHIBreadcrumbRange>& Ranges = QueueRanges.Add({ DeviceIndex, Pipeline });
+
+			// Extract the breadcrumb ranges for the payloads still in the interrupt queue.
+			for (FD3D12Payload* Payload : Queue.PendingInterrupt)
+			{
+				if (Payload->BreadcrumbRange)
+				{
+					Ranges.AddUnique(Payload->BreadcrumbRange);
+				}
+			}
+
+			BreadcrumbState.Devices[DeviceIndex].Pipelines[Pipeline].MarkerOut = Queue.DiagnosticBuffer->ReadMarkerOut();
+			BreadcrumbState.Devices[DeviceIndex].Pipelines[Pipeline].MarkerIn = Queue.DiagnosticBuffer->ReadMarkerIn();
+		});
+
+		// Traverse the breadcrumb tree and log active GPU work
+		if (!QueueRanges.IsEmpty())
+		{
+			BreadcrumbState.DumpActiveBreadcrumbs(QueueRanges, Verbosity);
+		}
+	}
+}
+#endif // WITH_RHI_BREADCRUMBS
+
 void FD3D12DynamicRHI::OutputGPUCrashReport(FTextBuilder& ErrorMessage)
 {
 	// Mark critical and gpu crash
@@ -846,74 +895,17 @@ void FD3D12DynamicRHI::OutputGPUCrashReport(FTextBuilder& ErrorMessage)
 		UE_LOG(LogD3D12RHI, Error, TEXT("GPU crash detected:%s\r\n"), *RemovedReasons);
 	}
 
-	// Log RHI breadcrumb data and shader asserts / prints
+
+	// Log shader asserts / prints
 	{
-#if WITH_RHI_BREADCRUMBS
-		TMap<FRHIBreadcrumbState::FQueueID, TArray<FRHIBreadcrumbRange>> QueueRanges;
-		FRHIBreadcrumbState BreadcrumbState;
-#endif
 		FString ShaderDiagnostics;
-
-		for (uint32 DeviceIndex = 0; DeviceIndex < GNumExplicitGPUsForRendering; ++DeviceIndex)
+		ForEachQueue([&](FD3D12Queue& Queue)
 		{
-			FD3D12Device* Device = GetAdapter().GetDevices()[DeviceIndex];
-			for (int32 QueueIndex = 0; QueueIndex < Device->GetQueues().Num(); ++QueueIndex)
-			{
-				FD3D12Queue& Queue = Device->GetQueues()[QueueIndex];
-
-				if (!Queue.DiagnosticBuffer)
-					continue;
-
-				const TCHAR* QueueName = GetD3DCommandQueueTypeName(ED3D12QueueType(QueueIndex));
-
-				ShaderDiagnostics += Queue.DiagnosticBuffer->GetShaderDiagnosticMessages(DeviceIndex, QueueIndex, QueueName);
-
-#if WITH_RHI_BREADCRUMBS
-				// Don't collect breadcrumb ranges if we don't have breadcrumbs enabled. The breadcrumbs
-				// will not have meaningful GPU state information because GPU markers are not written.
-				if (UE::RHI::UseGPUCrashBreadcrumbs())
-				{
-					ERHIPipeline Pipeline;
-					switch (Queue.QueueType)
-					{
-						default: continue; // Skip pipelines that the RHI doesn't handle
-						case ED3D12QueueType::Direct: Pipeline = ERHIPipeline::Graphics; break;
-						case ED3D12QueueType::Async: Pipeline = ERHIPipeline::AsyncCompute; break;
-					}
-
-					TArray<FRHIBreadcrumbRange>& Ranges = QueueRanges.Add({ DeviceIndex, Pipeline });
-
-					// Pull all incomplete payloads from the pending interrupt queue.
-					TArray<FD3D12Payload*> Payloads;
-					{
-						FD3D12Payload* Payload;
-						while (Queue.PendingInterrupt.Dequeue(Payload))
-							Payloads.Add(Payload);
-					}
-
-					// Extract the breadcrumb ranges for these payloads.
-					for (FD3D12Payload* Payload : Payloads)
-					{
-						if (Payload->BreadcrumbRange)
-						{
-							Ranges.AddUnique(Payload->BreadcrumbRange);
-						}
-					}
+			if (!Queue.DiagnosticBuffer)
+				return;
 
-					BreadcrumbState.Devices[DeviceIndex].Pipelines[Pipeline].MarkerOut = Queue.DiagnosticBuffer->ReadMarkerOut();
-					BreadcrumbState.Devices[DeviceIndex].Pipelines[Pipeline].MarkerIn = Queue.DiagnosticBuffer->ReadMarkerIn();
-				}
-#endif
-			}
-		}
-
-#if WITH_RHI_BREADCRUMBS
-		// Traverse the breadcrumb tree and log active GPU work
-		if (!QueueRanges.IsEmpty())
-		{
-			BreadcrumbState.DumpActiveBreadcrumbs(QueueRanges);
-		}
-#endif // WITH_RHI_BREADCRUMBS
+			ShaderDiagnostics += Queue.DiagnosticBuffer->GetShaderDiagnosticMessages(Queue.Device->GetGPUIndex(), Queue.QueueIndex, GetD3DCommandQueueTypeName(Queue.QueueType));
+		});
 
 		if (!ShaderDiagnostics.IsEmpty())
 		{
@@ -921,6 +913,16 @@ void FD3D12DynamicRHI::OutputGPUCrashReport(FTextBuilder& ErrorMessage)
 		}
 	}
 
+#if WITH_RHI_BREADCRUMBS
+	//
+	// Log RHI breadcrumb data
+	// 
+	// Don't collect breadcrumb ranges if we don't have breadcrumbs enabled. The breadcrumbs
+	// will not have meaningful GPU state information because GPU markers are not written.
+	//
+	DumpActiveBreadcrumbs(FRHIBreadcrumbState::EVerbosity::Error);
+#endif
+
 #if NV_AFTERMATH
 	TArray<UE::RHICore::Nvidia::Aftermath::FCrashResult> AftermathResults;
 	UE::RHICore::Nvidia::Aftermath::OnGPUCrash(AftermathResults);
diff --git a/Engine/Source/Runtime/RHI/Private/RHIBreadcrumbs.cpp b/Engine/Source/Runtime/RHI/Private/RHIBreadcrumbs.cpp
--- a/Engine/Source/Runtime/RHI/Private/RHIBreadcrumbs.cpp
+++ b/Engine/Source/Runtime/RHI/Private/RHIBreadcrumbs.cpp
@@ -120,7 +120,7 @@ RHI_API void FRHIBreadcrumbNode::WriteCrashData(struct FCrashContextExtendedWrit
 	UE_LOG(LogRHI, Error, TEXT("%s"), String);
 }
 
-RHI_API void FRHIBreadcrumbState::DumpActiveBreadcrumbs(TMap<FQueueID, TArray<FRHIBreadcrumbRange>> const& QueueRanges) const
+RHI_API void FRHIBreadcrumbState::DumpActiveBreadcrumbs(TMap<FQueueID, TArray<FRHIBreadcrumbRange>> const& QueueRanges, EVerbosity Verbosity) const
 {
 	FRHIBreadcrumb::FBuffer Buffer;
 	FGPUBreadcrumbCrashData CrashData(TEXT("RHI"));
@@ -336,7 +336,14 @@ RHI_API void FRHIBreadcrumbState::DumpActiveBreadcrumbs(TMap<FQueueID, TArray<FR
 	{
 		FGenericCrashContext::SetGPUBreadcrumbs(MoveTemp(CrashData));
 	}
-	UE_LOG(LogRHI, Error, TEXT("Active GPU breadcrumbs:%s\r\n"), *Tree);
+
+	switch (Verbosity)
+	{
+	case EVerbosity::Log    : UE_LOG(LogRHI, Log    , TEXT("Active GPU breadcrumbs:%s\r\n"), *Tree); break; 
+	case EVerbosity::Warning: UE_LOG(LogRHI, Warning, TEXT("Active GPU breadcrumbs:%s\r\n"), *Tree); break;
+	case EVerbosity::Error  : UE_LOG(LogRHI, Error  , TEXT("Active GPU breadcrumbs:%s\r\n"), *Tree); break;
+	}
+	
 }
 
 #endif // WITH_ADDITIONAL_CRASH_CONTEXTS
diff --git a/Engine/Source/Runtime/RHI/Public/RHIBreadcrumbs.h b/Engine/Source/Runtime/RHI/Public/RHIBreadcrumbs.h
--- a/Engine/Source/Runtime/RHI/Public/RHIBreadcrumbs.h
+++ b/Engine/Source/Runtime/RHI/Public/RHIBreadcrumbs.h
@@ -223,7 +223,14 @@ namespace UE::RHI::GPUProfiler
 			}
 		};
 
-		RHI_API void DumpActiveBreadcrumbs(TMap<FQueueID, TArray<FRHIBreadcrumbRange>> const& QueueRanges) const;
+		enum EVerbosity
+		{
+			Log,
+			Warning,
+			Error
+		};
+
+		RHI_API void DumpActiveBreadcrumbs(TMap<FQueueID, TArray<FRHIBreadcrumbRange>> const& QueueRanges, EVerbosity Verbosity = EVerbosity::Error) const;
 	};
 
 	struct FRHIBreadcrumbNode
