Clean up more resources on D3D12RHI shutdown. This is mainly for the global constant buffers inside command contexts, each of which has a bindless handle that needs to be released earlier during shutdown.

diff --git a/Engine/Source/Runtime/D3D12RHI/Private/D3D12Device.cpp b/Engine/Source/Runtime/D3D12RHI/Private/D3D12Device.cpp
--- a/Engine/Source/Runtime/D3D12RHI/Private/D3D12Device.cpp
+++ b/Engine/Source/Runtime/D3D12RHI/Private/D3D12Device.cpp
@@ -440,8 +440,20 @@ void FD3D12Device::SetupAfterDeviceCreation()
 	}
 }
 
+void FD3D12Queue::CleanupResources()
+{
+	ObjectPool.Contexts.CleanupResources();
+	ObjectPool.Allocators.CleanupResources();
+	ObjectPool.Lists.CleanupResources();
+}
+
 void FD3D12Device::CleanupResources()
 {
+	for (FD3D12Queue& Queue : Queues)
+	{
+		Queue.CleanupResources();
+	}
+
 	for (FD3D12OfflineDescriptorManager& Manager : OfflineDescriptorManagers)
 	{
 		Manager.CleanupResources();
diff --git a/Engine/Source/Runtime/D3D12RHI/Private/D3D12Device.h b/Engine/Source/Runtime/D3D12RHI/Private/D3D12Device.h
--- a/Engine/Source/Runtime/D3D12RHI/Private/D3D12Device.h
+++ b/Engine/Source/Runtime/D3D12RHI/Private/D3D12Device.h
@@ -103,6 +103,8 @@ public:
 class FD3D12Queue final
 {
 public:
+	void CleanupResources();
+
 	FD3D12Device* const Device;
 	ED3D12QueueType const QueueType;
 	int32 const QueueIndex;
diff --git a/Engine/Source/Runtime/D3D12RHI/Private/D3D12RHICommon.h b/Engine/Source/Runtime/D3D12RHI/Private/D3D12RHICommon.h
--- a/Engine/Source/Runtime/D3D12RHI/Private/D3D12RHICommon.h
+++ b/Engine/Source/Runtime/D3D12RHI/Private/D3D12RHICommon.h
@@ -379,6 +379,11 @@ template <typename TObjectType>
 struct TD3D12ObjectPool : public TLockFreePointerListUnordered<TObjectType, PLATFORM_CACHE_LINE_SIZE>
 {
 	~TD3D12ObjectPool()
+	{
+		CleanupResources();
+	}
+
+	void CleanupResources()
 	{
 		while (TObjectType* Object = TLockFreePointerListUnordered<TObjectType, PLATFORM_CACHE_LINE_SIZE>::Pop())
 		{
