D3D12 Enhanced Barriers

Add support for the Enhanced Barrier API.  This includes several notable changes:
  - Rework the D3D12 allocator logic to use ERHIAccess bits instead of D3D12_RESOURCE_STATES
  - Additionally, pass default and initial states through all the allocator paths so that we no longer rely on common state promotion in any case during resource creation
  - Add interface classes for the barrier implementations and factor the existing legacy barrier implementation into a implementation of those interfaces (this is why code moved around)
  - Add a factory to choose a barrier implementation and de-virtualize the function calls if only one is compiled in
  - Finally, implement Enhanced Barriers

The following defines determine if a platform, assuming it can support it, will include a given barrier implementation at compile time:
  - D3D12RHI_ALLOW_ENHANCED_BARRIERS
  - D3D12RHI_ALLOW_LEGACY_BARRIERS

Then at runtime, a preferred implementation type can be specified using the CVar r.D3D12.PreferredBarrierImplementation

For now we'll continue to prefer legacy barriers and only use enhanced barriers
on platforms where that's the only implementation included at compile time. This is due
to a hole in the API (and the related work-around) that means consistent behavior on
PC isn't guaranteed.

diff --git a/Engine/Source/Runtime/D3D12RHI/Private/D3D12Adapter.cpp b/Engine/Source/Runtime/D3D12RHI/Private/D3D12Adapter.cpp
--- a/Engine/Source/Runtime/D3D12RHI/Private/D3D12Adapter.cpp
+++ b/Engine/Source/Runtime/D3D12RHI/Private/D3D12Adapter.cpp
@@ -78,6 +78,52 @@ bool D3D12_ShouldBreakOnD3DDebugWarnings()
 	return GD3D12DebugCvar.GetValueOnAnyThread() == 3;
 }
 
+// 1 for now (e.g. EBs only on platforms where that's the only available implementation type)
+// There's a hole in the EB API that means we have to do some hacking to make it "work" and that's
+// not guaranteed to do the right thing on a wide range of drivers so we can't use them by default
+int32 GPreferredBarrierImplementation = 1;
+static FAutoConsoleVariableRef CVarPreferredBarrierImplementation(
+	TEXT("r.D3D12.PreferredBarrierImplementation"),
+	GPreferredBarrierImplementation,
+	TEXT("The preferred barrier implementation to use\n")
+	TEXT("0 Default\n")
+	TEXT("1 Legacy Barriers\n")
+	TEXT("2 Enhanced Barriers\n")
+	TEXT("All other values treated as Default"),
+	ECVF_ReadOnly);
+
+static ED3D12BarrierImplementationType GetImplementationTypeFromCVarValue(
+	int32 CVarValue)
+{
+	switch (CVarValue)
+	{
+	default:
+	case 0:
+		return ED3D12BarrierImplementationType::Invalid;
+	case 1:
+		return ED3D12BarrierImplementationType::Legacy;
+	case 2:
+		return ED3D12BarrierImplementationType::Enhanced;
+	}
+}
+
+static const TCHAR* GetImplementationTypeString(ED3D12BarrierImplementationType ImplType)
+{
+	switch (ImplType)
+	{
+	case ED3D12BarrierImplementationType::Legacy:
+		return TEXT("Legacy");
+	case ED3D12BarrierImplementationType::Enhanced:
+		return TEXT("Enhanced");
+	default:
+		checkNoEntry();
+		return TEXT("Unknown");
+	}
+}
+
+ED3D12BarrierImplementationType FD3D12Adapter::PreferredBarrierImplType = ED3D12BarrierImplementationType::Invalid;
+TUniquePtr<const FD3D12BarriersFactory::BarriersForAdapterType> FD3D12Adapter::Barriers = nullptr;
+
 #if PLATFORM_WINDOWS
 
 enum class ED3D12DredMode
@@ -1121,6 +1167,51 @@ void FD3D12Adapter::InitializeDevices()
 				GRHIGlobals.SupportsNative16BitOps = Features.Native16BitShaderOpsSupported;
 			}
 
+			bool bEnhancedBarriersSupported = false;
+#if D3D12RHI_SUPPORTS_ENHANCED_BARRIERS
+			{
+				D3D12_FEATURE_DATA_D3D12_OPTIONS12 options12 = {};
+				if (SUCCEEDED(RootDevice->CheckFeatureSupport(D3D12_FEATURE_D3D12_OPTIONS12, &options12, sizeof(options12))))
+				{
+					bEnhancedBarriersSupported = options12.EnhancedBarriersSupported;
+				}
+			}
+#endif
+			PreferredBarrierImplType =
+				[bEnhancedBarriersSupported]()
+				{
+					const ED3D12BarrierImplementationType CVarImplType =
+						GetImplementationTypeFromCVarValue(GPreferredBarrierImplementation);
+					switch (CVarImplType)
+					{
+					case ED3D12BarrierImplementationType::Invalid:
+						return bEnhancedBarriersSupported ?
+								ED3D12BarrierImplementationType::Enhanced :
+								ED3D12BarrierImplementationType::Legacy;
+					default:
+						return CVarImplType;
+					}
+				}();
+
+			const auto BoolToYesNo = [](bool bValue) { return bValue ? TEXT("Yes") : TEXT("No"); };
+			UE_LOG(LogD3D12RHI, Log, TEXT("     Legacy Barriers Present in Build: %s"), BoolToYesNo(!!D3D12RHI_SUPPORTS_LEGACY_BARRIERS));
+			UE_LOG(LogD3D12RHI, Log, TEXT("   Enhanced Barriers Present in Build: %s"), BoolToYesNo(!!D3D12RHI_SUPPORTS_ENHANCED_BARRIERS));
+			UE_LOG(LogD3D12RHI, Log, TEXT("     Enhanced Barriers Driver Support: %s"), BoolToYesNo(bEnhancedBarriersSupported));
+			UE_LOG(LogD3D12RHI, Log, TEXT("Preferred Barrier Implementation Type: %s"), GetImplementationTypeString(PreferredBarrierImplType));
+
+			check(!Barriers);
+			Barriers = TUniquePtr<FD3D12BarriersFactory::BarriersForAdapterType>(FD3D12BarriersFactory::CreateBarriersForAdapter(PreferredBarrierImplType));
+			check(Barriers);
+			if (!Barriers)
+			{
+				UE_LOG(LogD3D12RHI, Fatal, TEXT("Unable to find a D3D12 barrier implementation!"));
+			}
+			UE_LOG(LogD3D12RHI, Log, TEXT("        Chosen Barrier Implementation: %s"), Barriers->GetImplementationName());
+
+			GRHITransitionPrivateData_SizeInBytes = Barriers->GetTransitionDataSizeBytes();
+			GRHITransitionPrivateData_AlignInBytes = Barriers->GetTransitionDataAlignmentBytes();
+			Barriers->ConfigureDevice(RootDevice, bWithD3DDebug);
+
 #if D3D12_RHI_RAYTRACING
 #if PLATFORM_WINDOWS
 			D3D12_FEATURE_DATA_D3D12_OPTIONS5 D3D12Caps5 = {};
@@ -2098,4 +2189,22 @@ void FD3D12Adapter::SetResidencyPriority(ID3D12Pageable* Pageable, D3D12_RESIDEN
 #endif // D3D12_RHI_RAYTRACING
 }
 
+void FD3D12Adapter::CreateTransition(
+	FRHITransition* Transition,
+	const FRHITransitionCreateInfo& CreateInfo)
+{
+	check(Barriers);
+	return Barriers->CreateTransition(Transition, CreateInfo);
+}
 
+void FD3D12Adapter::ReleaseTransition(FRHITransition* Transition)
+{
+	check(Barriers);
+	return Barriers->ReleaseTransition(Transition);
+}
+
+TUniquePtr<FD3D12BarriersFactory::BarriersForContextType> FD3D12Adapter::CreateBarriersForContext()
+{
+	return TUniquePtr<FD3D12BarriersFactory::BarriersForContextType>(
+		FD3D12BarriersFactory::CreateBarriersForContext(PreferredBarrierImplType));
+}
diff --git a/Engine/Source/Runtime/D3D12RHI/Private/D3D12Adapter.h b/Engine/Source/Runtime/D3D12RHI/Private/D3D12Adapter.h
--- a/Engine/Source/Runtime/D3D12RHI/Private/D3D12Adapter.h
+++ b/Engine/Source/Runtime/D3D12RHI/Private/D3D12Adapter.h
@@ -38,6 +38,8 @@
 #include "D3D12RootSignature.h"
 #include "D3D12BindlessDescriptors.h"
 
+#include COMPILED_PLATFORM_HEADER(D3D12BarriersFactory.h)
+
 class FD3D12TransientHeapCache;
 class IRHITransientMemoryCache;
 
@@ -294,44 +296,82 @@ public:
 	HRESULT CreateCommittedResource(const FD3D12ResourceDesc& InDesc,
 		FRHIGPUMask CreationNode,
 		const D3D12_HEAP_PROPERTIES& HeapProps,
-		D3D12_RESOURCE_STATES InInitialState,
+		ED3D12Access InInitialD3D12Access,
 		const D3D12_CLEAR_VALUE* ClearValue,
 		FD3D12Resource** ppOutResource,
 		const TCHAR* Name,
 		bool bVerifyHResult = true)
 	{
-		return CreateCommittedResource(InDesc, CreationNode, HeapProps, InInitialState, ED3D12ResourceStateMode::Default, D3D12_RESOURCE_STATE_TBD, ClearValue, ppOutResource, Name, bVerifyHResult);
+		return CreateCommittedResource(
+			InDesc,
+			CreationNode,
+			HeapProps,
+			InInitialD3D12Access,
+			ED3D12ResourceStateMode::Default,
+			ED3D12Access::Unknown,
+			ClearValue,
+			ppOutResource,
+			Name,
+			bVerifyHResult);
 	}
 
 	HRESULT CreateCommittedResource(const FD3D12ResourceDesc& Desc,
 		FRHIGPUMask CreationNode,
 		const D3D12_HEAP_PROPERTIES& HeapProps,
-		D3D12_RESOURCE_STATES InInitialState,
+		ED3D12Access InInitialD3D12Access,
 		ED3D12ResourceStateMode InResourceStateMode,
-		D3D12_RESOURCE_STATES InDefaultState,
+		ED3D12Access InDefaultD3D12Access,
 		const D3D12_CLEAR_VALUE* ClearValue,
 		FD3D12Resource** ppOutResource,
 		const TCHAR* Name,
 		bool bVerifyHResult = true);
 
+	HRESULT CreateCommittedResourceRaw(
+		const D3D12_HEAP_PROPERTIES& InHeapProps,
+		D3D12_HEAP_FLAGS InHeapFlags,
+		const FD3D12ResourceDesc& InDesc,
+		ED3D12Access InInitialD3D12Access,
+		const D3D12_CLEAR_VALUE* InClearValue,
+		TRefCountPtr<ID3D12Resource>& OutResource)
+	{
+		return Barriers->CreateCommittedResource(
+			*this,
+			InHeapProps,
+			InHeapFlags,
+			InDesc,
+			InInitialD3D12Access,
+			InClearValue,
+			OutResource);
+	}
+
 	HRESULT CreatePlacedResource(const FD3D12ResourceDesc& InDesc,
 		FD3D12Heap* BackingHeap,
 		uint64 HeapOffset,
-		D3D12_RESOURCE_STATES InInitialState,
+		ED3D12Access InInitialD3D12Access,
 		const D3D12_CLEAR_VALUE* ClearValue,
 		FD3D12Resource** ppOutResource,
 		const TCHAR* Name,
 		bool bVerifyHResult = true)
 	{
-		return CreatePlacedResource(InDesc, BackingHeap, HeapOffset, InInitialState, ED3D12ResourceStateMode::Default, D3D12_RESOURCE_STATE_TBD, ClearValue, ppOutResource, Name, bVerifyHResult);
+		return CreatePlacedResource(
+			InDesc,
+			BackingHeap,
+			HeapOffset,
+			InInitialD3D12Access,
+			ED3D12ResourceStateMode::Default,
+			InInitialD3D12Access,
+			ClearValue,
+			ppOutResource,
+			Name,
+			bVerifyHResult);
 	}
 
 	HRESULT CreatePlacedResource(const FD3D12ResourceDesc& Desc,
 		FD3D12Heap* BackingHeap,
 		uint64 HeapOffset,
-		D3D12_RESOURCE_STATES InInitialState,
+		ED3D12Access InInitialD3D12Access,
 		ED3D12ResourceStateMode InResourceStateMode,
-		D3D12_RESOURCE_STATES InDefaultState,
+		ED3D12Access InDefaultD3D12Access,
 		const D3D12_CLEAR_VALUE* ClearValue,
 		FD3D12Resource** ppOutResource,
 		const TCHAR* Name,
@@ -339,9 +379,9 @@ public:
 
 	HRESULT CreateReservedResource(const FD3D12ResourceDesc& Desc,
 		FRHIGPUMask CreationNode,
-		D3D12_RESOURCE_STATES InInitialState,
+		ED3D12Access InInitialD3D12Access,
 		ED3D12ResourceStateMode InResourceStateMode,
-		D3D12_RESOURCE_STATES InDefaultState,
+		ED3D12Access InDefaultD3D12Access,
 		const D3D12_CLEAR_VALUE* ClearValue,
 		FD3D12Resource** ppOutResource,
 		const TCHAR* Name,
@@ -358,8 +398,9 @@ public:
 	HRESULT CreateBuffer(D3D12_HEAP_TYPE HeapType,
 		FRHIGPUMask CreationNode,
 		FRHIGPUMask VisibleNodes,
-		D3D12_RESOURCE_STATES InitialState,
+		ED3D12Access InInitialD3D12Access,
 		ED3D12ResourceStateMode InResourceStateMode,
+		ED3D12Access InDefaultD3D12Access,
 		uint64 HeapSize,
 		FD3D12Resource** ppOutResource,
 		const TCHAR* Name,
@@ -367,9 +408,9 @@ public:
 
 	HRESULT CreateBuffer(const D3D12_HEAP_PROPERTIES& HeapProps,
 		FRHIGPUMask CreationNode,
-		D3D12_RESOURCE_STATES InInitialState,
+		ED3D12Access InInitialD3D12Access,
 		ED3D12ResourceStateMode InResourceStateMode,
-		D3D12_RESOURCE_STATES InDefaultState,
+		ED3D12Access InDefaultD3D12Access,
 		uint64 HeapSize,
 		FD3D12Resource** ppOutResource,
 		const TCHAR* Name,
@@ -380,7 +421,7 @@ public:
 		uint32 Alignment,
 		const FRHIBufferCreateDesc& CreateDesc,
 		ED3D12ResourceStateMode InResourceStateMode,
-		D3D12_RESOURCE_STATES InCreateState,
+		ED3D12Access InCreateD3D12Access,
 		bool bKeepUnlocked = false,
 		ID3D12ResourceAllocator* ResourceAllocator = nullptr
 	);
@@ -451,6 +492,14 @@ public:
 	FD3D12BindlessDescriptorAllocator& GetBindlessDescriptorAllocator() { return BindlessDescriptorAllocator; }
 #endif
 
+	void CreateTransition(
+		FRHITransition* Transition,
+		const FRHITransitionCreateInfo& CreateInfo);
+	void ReleaseTransition(FRHITransition* Transition);
+
+	[[nodiscard]]
+	static TUniquePtr<FD3D12BarriersFactory::BarriersForContextType> CreateBarriersForContext();
+
 protected:
 
 	virtual void CreateRootDevice(bool bWithDebug);
@@ -462,7 +511,7 @@ protected:
 		uint32 Size,
 		EBufferUsageFlags InUsage,
 		ED3D12ResourceStateMode InResourceStateMode,
-		D3D12_RESOURCE_STATES InCreateState,
+		ED3D12Access InCreatED3D12Access,
 		uint32 Alignment,
 		FD3D12Buffer* Buffer,
 		FD3D12ResourceLocation& ResourceLocation,
@@ -470,6 +519,9 @@ protected:
 		const TCHAR* InDebugName,
 		const FName& OwnerName);
 
+	static ED3D12BarrierImplementationType PreferredBarrierImplType;
+	static TUniquePtr<const FD3D12BarriersFactory::BarriersForAdapterType> Barriers;
+
 	// Creates default root and execute indirect signatures
 	virtual void CreateCommandSignatures();
 
diff --git a/Engine/Source/Runtime/D3D12RHI/Private/D3D12Allocation.cpp b/Engine/Source/Runtime/D3D12RHI/Private/D3D12Allocation.cpp
--- a/Engine/Source/Runtime/D3D12RHI/Private/D3D12Allocation.cpp
+++ b/Engine/Source/Runtime/D3D12RHI/Private/D3D12Allocation.cpp
@@ -7,6 +7,7 @@
 //-----------------------------------------------------------------------------
 #include "D3D12Allocation.h"
 #include "D3D12RHIPrivate.h"
+#include "D3D12Util.h"
 #include "Misc/BufferedOutputDevice.h"
 #include "HAL/PlatformStackWalk.h"
 #include "HAL/LowLevelMemTracker.h"
@@ -308,11 +309,31 @@ void FD3D12BuddyAllocator::Initialize()
 			const D3D12_HEAP_PROPERTIES HeapProps = CD3DX12_HEAP_PROPERTIES(InitConfig.HeapType, GetGPUMask().GetNative(), GetVisibilityMask().GetNative());
 #if PLATFORM_WINDOWS
 			LLM_SCOPED_PAUSE_TRACKING_FOR_TRACKER(ELLMTracker::Default, ELLMAllocType::System);
-			VERIFYD3D12RESULT(Adapter->CreateBuffer(HeapProps, GetGPUMask(), InitConfig.InitialResourceState, ED3D12ResourceStateMode::SingleState, InitConfig.InitialResourceState, MaxBlockSize, BackingResource.GetInitReference(), TEXT("Resource Allocator Underlying Buffer"), InitConfig.ResourceFlags));
+			VERIFYD3D12RESULT(
+				Adapter->CreateBuffer(
+					HeapProps,
+					GetGPUMask(),
+					InitConfig.InitialD3D12Access,
+					ED3D12ResourceStateMode::SingleState,
+					InitConfig.InitialD3D12Access,
+					MaxBlockSize,
+					BackingResource.GetInitReference(),
+					TEXT("Resource Allocator Underlying Buffer"),
+					InitConfig.ResourceFlags));
 			LLM_SCOPED_PAUSE_TRACKING_WITH_ENUM_AND_AMOUNT_BYTAG(D3D12AllocatorUnused, MaxBlockSize, ELLMTracker::Platform, ELLMAllocType::System);
 #else
 			LLM_SCOPE_BYTAG(D3D12AllocatorUnused);
-			VERIFYD3D12RESULT(Adapter->CreateBuffer(HeapProps, GetGPUMask(), InitConfig.InitialResourceState, ED3D12ResourceStateMode::SingleState, InitConfig.InitialResourceState, MaxBlockSize, BackingResource.GetInitReference(), TEXT("Resource Allocator Underlying Buffer"), InitConfig.ResourceFlags));
+			VERIFYD3D12RESULT(
+				Adapter->CreateBuffer(
+					HeapProps,
+					GetGPUMask(),
+					InitConfig.InitialD3D12Access,
+					ED3D12ResourceStateMode::SingleState,
+					InitConfig.InitialD3D12Access,
+					MaxBlockSize,
+					BackingResource.GetInitReference(),
+					TEXT("Resource Allocator Underlying Buffer"),
+					InitConfig.ResourceFlags));
 #endif
 #if UE_MEMORY_TRACE_ENABLED
 			MemoryTrace_MarkAllocAsHeap(BackingResource->GetGPUVirtualAddress(), TraceHeapId);
@@ -1144,7 +1165,18 @@ void* FD3D12UploadHeapAllocator::AllocUploadResource(uint32 InSize, uint32 InAli
 
 		// Forward to the big block allocator
 		const D3D12_RESOURCE_DESC ResourceDesc = CD3DX12_RESOURCE_DESC::Buffer(InSize, D3D12_RESOURCE_FLAG_NONE);
-		BigBlockAllocator.AllocateResource(GetParentDevice()->GetGPUIndex(), D3D12_HEAP_TYPE_UPLOAD, ResourceDesc, InSize, InAlignment, ED3D12ResourceStateMode::SingleState, D3D12_RESOURCE_STATE_GENERIC_READ, nullptr, nullptr, ResourceLocation);
+		BigBlockAllocator.AllocateResource(
+			GetParentDevice()->GetGPUIndex(),
+			D3D12_HEAP_TYPE_UPLOAD,
+			ResourceDesc,
+			InSize,
+			InAlignment,
+			ED3D12Access::GenericRead,
+			ED3D12ResourceStateMode::SingleState,
+			ED3D12Access::GenericRead,
+			nullptr,
+			nullptr,
+			ResourceLocation);
 		ResourceLocation.UnlockPoolData();		
 	}
 	
@@ -1216,24 +1248,26 @@ FD3D12ResourceInitConfig FD3D12DefaultBufferPool::GetResourceAllocatorInitConfig
 	// Setup initial resource state depending on the requested buffer flags
 	if (EnumHasAnyFlags(InBufferUsage, BUF_AccelerationStructure))
 	{
+		// TODO: Can we do better than READ + WRITE all the time?
 		// should only have this flag and no other flags
 		check(InBufferUsage == BUF_AccelerationStructure);
-		InitConfig.InitialResourceState = D3D12_RESOURCE_STATE_RAYTRACING_ACCELERATION_STRUCTURE;
+		check(EnumHasAnyFlags(InResourceFlags, D3D12_RESOURCE_FLAG_RAYTRACING_ACCELERATION_STRUCTURE));
+		InitConfig.InitialD3D12Access = ED3D12Access::BVHRead | ED3D12Access::BVHWrite;
 	}
 	else 
 #endif // D3D12_RHI_RAYTRACING
 	if (InitConfig.HeapType == D3D12_HEAP_TYPE_READBACK)
 	{
-		InitConfig.InitialResourceState = D3D12_RESOURCE_STATE_COPY_DEST;
+		InitConfig.InitialD3D12Access = ED3D12Access::CopyDest;
 	}
 	else if (EnumHasAnyFlags(InBufferUsage, BUF_UnorderedAccess))
 	{
 		check(InResourceFlags & D3D12_RESOURCE_FLAG_ALLOW_UNORDERED_ACCESS);
-		InitConfig.InitialResourceState = D3D12_RESOURCE_STATE_UNORDERED_ACCESS;
+		InitConfig.InitialD3D12Access = ED3D12Access::UAVMask;
 	}
 	else
 	{
-		InitConfig.InitialResourceState = D3D12_RESOURCE_STATE_GENERIC_READ;
+		InitConfig.InitialD3D12Access = ED3D12Access::GenericRead;
 	}
 
 	InitConfig.HeapFlags = D3D12_HEAP_FLAG_ALLOW_ONLY_BUFFERS;
@@ -1293,8 +1327,15 @@ void FD3D12DefaultBufferPool::CleanUpAllocations(uint64 FrameLag)
 }
 
 // Grab a buffer from the available buffers or create a new buffer if none are available
-void FD3D12DefaultBufferPool::AllocDefaultResource(D3D12_HEAP_TYPE InHeapType, const D3D12_RESOURCE_DESC& Desc, EBufferUsageFlags InUsage, ED3D12ResourceStateMode InResourceStateMode,
-	D3D12_RESOURCE_STATES InCreateState, uint32 Alignment, const TCHAR* Name, FD3D12ResourceLocation& ResourceLocation)
+void FD3D12DefaultBufferPool::AllocDefaultResource(
+	D3D12_HEAP_TYPE InHeapType,
+	const D3D12_RESOURCE_DESC& Desc,
+	EBufferUsageFlags InUsage,
+	ED3D12ResourceStateMode InResourceStateMode,
+	ED3D12Access InCreateD3D12Access,
+	uint32 Alignment,
+	const TCHAR* Name,
+	FD3D12ResourceLocation& ResourceLocation)
 {
 	FD3D12Device* Device = GetParentDevice();
 	FD3D12Adapter* Adapter = Device->GetParentAdapter();
@@ -1311,18 +1352,21 @@ void FD3D12DefaultBufferPool::AllocDefaultResource(D3D12_HEAP_TYPE InHeapType, c
 	// Validate the create state
 	if (InHeapType == D3D12_HEAP_TYPE_READBACK)
 	{
-		check(InCreateState == D3D12_RESOURCE_STATE_COPY_DEST);
+		check(InCreateD3D12Access == ED3D12Access::CopyDest);
 	}
 	else if (InHeapType == D3D12_HEAP_TYPE_UPLOAD)
 	{
-		check(InCreateState == D3D12_RESOURCE_STATE_GENERIC_READ);
+		check(InCreateD3D12Access == ED3D12Access::GenericRead);
 	}
 #if D3D12_RHI_RAYTRACING
 	else if (EnumHasAnyFlags(InUsage, BUF_AccelerationStructure))
 	{
+		// @TODO - This isn't true anymore when using enhanced barriers and when using legacy barriers, it's the 
+		// implementations responsibility now to deal with the API restriction. This can be re-written to use
+		// BVHRead and BVHWrite as appropriate so that barriers implementations that can do better, will.
 		// RayTracing acceleration structures must be created in a particular state and may never transition out of it.
 		check(InResourceStateMode == ED3D12ResourceStateMode::SingleState);
-		check(InCreateState == D3D12_RESOURCE_STATE_RAYTRACING_ACCELERATION_STRUCTURE);
+		check(EnumHasAnyFlags(InCreateD3D12Access, ED3D12Access::BVHRead | ED3D12Access::BVHWrite));
 	}
 #endif // D3D12_RHI_RAYTRACING
 #endif  // DO_CHECK
@@ -1355,7 +1399,17 @@ void FD3D12DefaultBufferPool::AllocDefaultResource(D3D12_HEAP_TYPE InHeapType, c
 				uint64 HeapOffset = ResourceLocation.GetAllocator()->GetAllocationOffsetInBytes(ResourceLocation.GetBuddyAllocatorPrivateData());
 
 				FD3D12Resource* NewResource = nullptr;
-				VERIFYD3D12RESULT(Adapter->CreatePlacedResource(Desc, BackingHeap, HeapOffset, InCreateState, ED3D12ResourceStateMode::MultiState, D3D12_RESOURCE_STATE_TBD, nullptr, &NewResource, Name));
+				VERIFYD3D12RESULT(Adapter->CreatePlacedResource(
+					Desc,
+					BackingHeap,
+					HeapOffset,
+					InCreateD3D12Access,
+					// @TODO - Is this right? We're ignoring InResourceStateMode
+					ED3D12ResourceStateMode::MultiState,
+					ED3D12Access::Unknown,
+					nullptr,
+					&NewResource,
+					Name));
 
 				ResourceLocation.SetResource(NewResource);
 			}
@@ -1372,7 +1426,17 @@ void FD3D12DefaultBufferPool::AllocDefaultResource(D3D12_HEAP_TYPE InHeapType, c
 	// Allocate Standalone
 	// Todo: track stand alone allocations and see how much memory we use by this and how many we have
 	FD3D12Resource* NewResource = nullptr;
-	VERIFYD3D12RESULT(Adapter->CreateBuffer(InHeapType, GetGPUMask(), GetVisibilityMask(), InCreateState, InResourceStateMode, Desc.Width, &NewResource, Name, Desc.Flags));
+	VERIFYD3D12RESULT(Adapter->CreateBuffer(
+		InHeapType,
+		GetGPUMask(),
+		GetVisibilityMask(),
+		InCreateD3D12Access,
+		InResourceStateMode,
+		InCreateD3D12Access,
+		Desc.Width,
+		&NewResource,
+		Name,
+		Desc.Flags));
 
 	ResourceLocation.AsStandAlone(NewResource, Desc.Width);
 }
@@ -1423,7 +1487,7 @@ FD3D12BufferPool* FD3D12DefaultBufferAllocator::CreateBufferPool(D3D12_HEAP_TYPE
 	// Disable defrag on the RT Acceleration pool - #kenzo_todo
 	// Acceleration structure buffers may be created, but not built immediately.
 	// If CopyRaytracingAccelerationStructure is called on such buffer, the GPU will crash.
-	if (InitConfig.InitialResourceState == D3D12_RESOURCE_STATE_RAYTRACING_ACCELERATION_STRUCTURE)
+	if (EnumHasAnyFlags(InitConfig.InitialD3D12Access, ED3D12Access::BVHRead | ED3D12Access::BVHWrite))
 	{
 		bDefragEnabled = false;
 
@@ -1470,38 +1534,49 @@ bool FD3D12DefaultBufferAllocator::IsPlacedResource(D3D12_RESOURCE_FLAGS InResou
 }
 
 
-D3D12_RESOURCE_STATES FD3D12DefaultBufferAllocator::GetDefaultInitialResourceState(D3D12_HEAP_TYPE InHeapType, EBufferUsageFlags InBufferFlags, ED3D12ResourceStateMode InResourceStateMode)
+ED3D12Access FD3D12DefaultBufferAllocator::GetDefaultInitialD3D12Access(
+	D3D12_HEAP_TYPE InHeapType,
+	EBufferUsageFlags InBufferFlags,
+	ED3D12ResourceStateMode InResourceStateMode)
 {
 	// Validate the create state
 	if (InHeapType == D3D12_HEAP_TYPE_READBACK)
 	{
-		return D3D12_RESOURCE_STATE_COPY_DEST;
+		return ED3D12Access::CopyDest;
 	}
 	else if (InHeapType == D3D12_HEAP_TYPE_UPLOAD)
 	{
-		return D3D12_RESOURCE_STATE_GENERIC_READ;
+		return ED3D12Access::GenericRead;
 	}
 	else if (InBufferFlags == BUF_UnorderedAccess && InResourceStateMode == ED3D12ResourceStateMode::SingleState)
 	{
 		check(InHeapType == D3D12_HEAP_TYPE_DEFAULT);
-		return D3D12_RESOURCE_STATE_UNORDERED_ACCESS;
+		return ED3D12Access::UAVMask;
 	}
 #if D3D12_RHI_RAYTRACING
 	else if (EnumHasAnyFlags(InBufferFlags, BUF_AccelerationStructure))
 	{
 		check(InHeapType == D3D12_HEAP_TYPE_DEFAULT);
-		return D3D12_RESOURCE_STATE_RAYTRACING_ACCELERATION_STRUCTURE;
+		return ED3D12Access::BVHRead | ED3D12Access::BVHWrite;
 	}
 #endif // D3D12_RHI_RAYTRACING
 	else
 	{
-		return D3D12_RESOURCE_STATE_GENERIC_READ;
+		return ED3D12Access::GenericRead;
 	}
 }
 
 
 // Grab a buffer from the available buffers or create a new buffer if none are available
-void FD3D12DefaultBufferAllocator::AllocDefaultResource(D3D12_HEAP_TYPE InHeapType, const D3D12_RESOURCE_DESC& InResourceDesc, EBufferUsageFlags InBufferUsage, ED3D12ResourceStateMode InResourceStateMode, D3D12_RESOURCE_STATES InCreateState, FD3D12ResourceLocation& ResourceLocation, uint32 Alignment, const TCHAR* Name)
+void FD3D12DefaultBufferAllocator::AllocDefaultResource(
+	D3D12_HEAP_TYPE InHeapType,
+	const D3D12_RESOURCE_DESC& InResourceDesc,
+	EBufferUsageFlags InBufferUsage,
+	ED3D12ResourceStateMode InResourceStateMode,
+	ED3D12Access InCreateD3D12Access,
+	FD3D12ResourceLocation& ResourceLocation,
+	uint32 Alignment,
+	const TCHAR* Name)
 {
 	FD3D12Adapter* Adapter = GetParentDevice()->GetParentAdapter();
 
@@ -1515,7 +1590,17 @@ void FD3D12DefaultBufferAllocator::AllocDefaultResource(D3D12_HEAP_TYPE InHeapTy
 		const D3D12_HEAP_PROPERTIES HeapProps = CD3DX12_HEAP_PROPERTIES(InHeapType, GetGPUMask().GetNative(), EnumHasAnyFlags(InBufferUsage, EBufferUsageFlags::NNE) ? FRHIGPUMask::GPU0().GetNative() : GetVisibilityMask().GetNative());
 		D3D12_RESOURCE_DESC Desc = InResourceDesc;
 		Desc.Alignment = 0;
-		VERIFYD3D12RESULT(Adapter->CreateCommittedResource(Desc, GetGPUMask(), HeapProps, InCreateState, InResourceStateMode, InCreateState, nullptr, &NewResource, Name, false));
+		VERIFYD3D12RESULT(Adapter->CreateCommittedResource(
+			Desc,
+			GetGPUMask(),
+			HeapProps,
+			InCreateD3D12Access,
+			InResourceStateMode,
+			InCreateD3D12Access,
+			nullptr,
+			&NewResource,
+			Name,
+			false));
 
 		ResourceLocation.AsStandAlone(NewResource, InResourceDesc.Width);
 
@@ -1532,7 +1617,16 @@ void FD3D12DefaultBufferAllocator::AllocDefaultResource(D3D12_HEAP_TYPE InHeapTy
 		FD3D12ResourceDesc Desc = InResourceDesc;
 		Desc.Alignment = Alignment;
 		Desc.bReservedResource = true;
-		VERIFYD3D12RESULT(Adapter->CreateReservedResource(Desc, GetGPUMask(), InCreateState, InResourceStateMode, InCreateState, nullptr, &NewResource, Name, false));
+		VERIFYD3D12RESULT(Adapter->CreateReservedResource(
+			Desc,
+			GetGPUMask(),
+			InCreateD3D12Access,
+			InResourceStateMode,
+			InCreateD3D12Access,
+			nullptr,
+			&NewResource,
+			Name,
+			false));
 
 		ResourceLocation.AsStandAlone(NewResource, InResourceDesc.Width);
 
@@ -1564,7 +1658,15 @@ void FD3D12DefaultBufferAllocator::AllocDefaultResource(D3D12_HEAP_TYPE InHeapTy
 	}
 
 	// Perform actual allocation
-	BufferPool->AllocDefaultResource(InHeapType, ResourceDesc, InBufferUsage, InResourceStateMode, InCreateState, Alignment, Name, ResourceLocation);
+	BufferPool->AllocDefaultResource(
+		InHeapType,
+		ResourceDesc,
+		InBufferUsage,
+		InResourceStateMode,
+		InCreateD3D12Access,
+		Alignment,
+		Name,
+		ResourceLocation);
 }
 
 
@@ -1700,7 +1802,7 @@ FD3D12TextureAllocatorPool::FD3D12TextureAllocatorPool(FD3D12Device* Device, FRH
 
 	// unused for textures because placed and not suballocated
 	SharedInitConfig.ResourceFlags = D3D12_RESOURCE_FLAG_NONE;
-	SharedInitConfig.InitialResourceState = D3D12_RESOURCE_STATE_COMMON;
+	SharedInitConfig.InitialD3D12Access = ED3D12Access::Common;
 
 	EResourceAllocationStrategy AllocationStrategy = EResourceAllocationStrategy::kPlacedResource;
 	FRHIMemoryPool::EFreeListOrder FreeListOrder = FRHIMemoryPool::EFreeListOrder::SortBySize;
@@ -1779,7 +1881,8 @@ HRESULT FD3D12TextureAllocatorPool::AllocateTexture(
 	const D3D12_CLEAR_VALUE* ClearValue,
 	EPixelFormat UEFormat,
 	FD3D12ResourceLocation& TextureLocation,
-	const D3D12_RESOURCE_STATES InitialState,
+	ED3D12Access InInitialD3D12Access,
+	ED3D12Access InDefaultD3D12Access,
 	const TCHAR* Name)
 {
 	// The top mip level must be less than 64 KB to use 4 KB alignment
@@ -1793,7 +1896,19 @@ HRESULT FD3D12TextureAllocatorPool::AllocateTexture(
 
 	ED3D12ResourceStateMode ResourceStateMode = bIsReadOnly ? ED3D12ResourceStateMode::Default : ED3D12ResourceStateMode::MultiState;
 	EPoolType PoolType = bIsReadOnly ? (b4KAligment ? EPoolType::ReadOnly4K : EPoolType::ReadOnly) : (bIsRenderTarget ? EPoolType::RenderTarget : EPoolType::UAV);
-	PoolAllocators[(int)PoolType]->AllocateResource(GetParentDevice()->GetGPUIndex(), D3D12_HEAP_TYPE_DEFAULT, Desc, Info.SizeInBytes, Info.Alignment, ResourceStateMode, InitialState, ClearValue, Name, TextureLocation);
+
+	PoolAllocators[(int)PoolType]->AllocateResource(
+		GetParentDevice()->GetGPUIndex(),
+		D3D12_HEAP_TYPE_DEFAULT,
+		Desc,
+		Info.SizeInBytes,
+		Info.Alignment,
+		InInitialD3D12Access,
+		ResourceStateMode,
+		InDefaultD3D12Access,
+		ClearValue,
+		Name,
+		TextureLocation);
 
 	return S_OK;
 }
@@ -1943,7 +2058,7 @@ FD3D12TextureAllocator::FD3D12TextureAllocator(FD3D12Device* Device,
 			D3D12_HEAP_TYPE_DEFAULT,
 			Flags | D3D12_HEAP_FLAG_DENY_BUFFERS,
 			D3D12_RESOURCE_FLAG_NONE,
-			D3D12_RESOURCE_STATE_GENERIC_READ
+			ED3D12Access::GenericRead,
 		},
 		Name,
 		EResourceAllocationStrategy::kPlacedResource,
@@ -1958,7 +2073,13 @@ FD3D12TextureAllocator::~FD3D12TextureAllocator()
 {
 }
 
-HRESULT FD3D12TextureAllocator::AllocateTexture(FD3D12ResourceDesc Desc, const D3D12_CLEAR_VALUE* ClearValue, FD3D12ResourceLocation& TextureLocation, const D3D12_RESOURCE_STATES InitialState, const TCHAR* Name)
+HRESULT FD3D12TextureAllocator::AllocateTexture(
+	FD3D12ResourceDesc Desc,
+	const D3D12_CLEAR_VALUE* ClearValue,
+	FD3D12ResourceLocation& TextureLocation,
+	ED3D12Access InInitialD3D12Access,
+	ED3D12Access InDefaultD3D12Access,
+	const TCHAR* Name)
 {
 	FD3D12Device* Device = GetParentDevice();
 	FD3D12Adapter* Adapter = Device->GetParentAdapter();
@@ -1977,7 +2098,17 @@ HRESULT FD3D12TextureAllocator::AllocateTexture(FD3D12ResourceDesc Desc, const D
 			FD3D12Heap* BackingHeap = ((FD3D12BuddyAllocator*)TextureLocation.GetAllocator())->GetBackingHeap();
 			uint64 HeapOffset = TextureLocation.GetAllocator()->GetAllocationOffsetInBytes(TextureLocation.GetBuddyAllocatorPrivateData());
 
-			hr = Adapter->CreatePlacedResource(Desc, BackingHeap, HeapOffset, InitialState, ClearValue, &NewResource, Name, false);
+			hr = Adapter->CreatePlacedResource(
+				Desc,
+				BackingHeap,
+				HeapOffset,
+				InInitialD3D12Access,
+				ED3D12ResourceStateMode::Default,
+				InDefaultD3D12Access,
+				ClearValue,
+				&NewResource,
+				Name,
+				false);
 	
 			TextureLocation.SetType(FD3D12ResourceLocation::ResourceLocationType::eSubAllocation);
 			TextureLocation.SetResource(NewResource);
@@ -1990,7 +2121,15 @@ HRESULT FD3D12TextureAllocator::AllocateTexture(FD3D12ResourceDesc Desc, const D
 	Desc.Alignment = 0;
 	const D3D12_HEAP_PROPERTIES HeapProps = CD3DX12_HEAP_PROPERTIES(D3D12_HEAP_TYPE_DEFAULT, GetGPUMask().GetNative(), (uint32)GetVisibilityMask().GetNative());
 
-	hr = Adapter->CreateCommittedResource(Desc, GetGPUMask(), HeapProps, InitialState, ClearValue, &NewResource, Name, false);
+	hr = Adapter->CreateCommittedResource(
+		Desc,
+		GetGPUMask(),
+		HeapProps,
+		InInitialD3D12Access,
+		ClearValue,
+		&NewResource,
+		Name,
+		false);
 
 	TextureLocation.AsStandAlone(NewResource, Info.SizeInBytes);
 
@@ -2004,12 +2143,25 @@ FD3D12TextureAllocatorPool::FD3D12TextureAllocatorPool(FD3D12Device* Device, FRH
 	ReadOnlyTexturePool(Device, VisibilityNode, FString(L"Small Non-DS/RT Texture allocator"), TEXTURE_POOL_SIZE, D3D12_HEAP_FLAG_ALLOW_ONLY_NON_RT_DS_TEXTURES, TraceHeapId)
 {};
 
-HRESULT FD3D12TextureAllocatorPool::AllocateTexture(FD3D12ResourceDesc Desc, const D3D12_CLEAR_VALUE* ClearValue, EPixelFormat UEFormat, FD3D12ResourceLocation& TextureLocation, const D3D12_RESOURCE_STATES InitialState, const TCHAR* Name)
+HRESULT FD3D12TextureAllocatorPool::AllocateTexture(
+	FD3D12ResourceDesc Desc,
+	const D3D12_CLEAR_VALUE* ClearValue,
+	EPixelFormat UEFormat,
+	FD3D12ResourceLocation& TextureLocation, 
+	ED3D12Access InInitialD3D12Access,
+	ED3D12Access InDefaultD3D12Access,
+	const TCHAR* Name)
 {
 	if (FD3D12Texture::CanBe4KAligned(Desc, (EPixelFormat)UEFormat))
 	{
 		Desc.Alignment = D3D12_SMALL_RESOURCE_PLACEMENT_ALIGNMENT; // request 4k alignment
-		return ReadOnlyTexturePool.AllocateTexture(Desc, ClearValue, TextureLocation, InitialState, Name);
+			return ReadOnlyTexturePool.AllocateTexture(
+				Desc,
+				ClearValue,
+				TextureLocation,
+				InInitialD3D12Access,
+				InDefaultD3D12Access,
+				Name);
 	}
 
 	FD3D12Adapter* Adapter = GetParentDevice()->GetParentAdapter();
@@ -2018,7 +2170,7 @@ HRESULT FD3D12TextureAllocatorPool::AllocateTexture(FD3D12ResourceDesc Desc, con
 	const D3D12_HEAP_PROPERTIES HeapProps = CD3DX12_HEAP_PROPERTIES(D3D12_HEAP_TYPE_DEFAULT, GetGPUMask().GetNative(), GetVisibilityMask().GetNative());
 	const D3D12_RESOURCE_ALLOCATION_INFO Info = GetParentDevice()->GetResourceAllocationInfoUncached(Desc);
 
-	HRESULT hr = Adapter->CreateCommittedResource(Desc, GetGPUMask(), HeapProps, InitialState, ClearValue, &Resource, Name, false);
+	HRESULT hr = Adapter->CreateCommittedResource(Desc, GetGPUMask(), HeapProps, InInitialD3D12Access, ClearValue, &Resource, Name, false);
 
 	if (SUCCEEDED(hr))
 	{
@@ -2173,8 +2325,16 @@ FD3D12FastAllocatorPage* FD3D12FastAllocatorPagePool::RequestFastAllocatorPage()
 
 	FD3D12FastAllocatorPage* Page = new FD3D12FastAllocatorPage(PageSize);
 
-	const D3D12_RESOURCE_STATES InitialState = DetermineInitialResourceState(HeapProperties.Type, &HeapProperties);
-	VERIFYD3D12RESULT(Adapter->CreateBuffer(HeapProperties, GetGPUMask(), InitialState, ED3D12ResourceStateMode::SingleState, InitialState, PageSize, Page->FastAllocBuffer.GetInitReference(), TEXT("Fast Allocator Page")));
+	const ED3D12Access InitialD3D12Access = DetermineInitialBufferD3D12Access(HeapProperties.Type, &HeapProperties);
+	VERIFYD3D12RESULT(Adapter->CreateBuffer(
+		HeapProperties,
+		GetGPUMask(),
+		InitialD3D12Access,
+		ED3D12ResourceStateMode::SingleState,
+		InitialD3D12Access,
+		PageSize,
+		Page->FastAllocBuffer.GetInitReference(),
+		TEXT("Fast Allocator Page")));
 	Page->FastAllocBuffer->DoNotDeferDelete();
 
 	Page->FastAllocData = Page->FastAllocBuffer->Map();
diff --git a/Engine/Source/Runtime/D3D12RHI/Private/D3D12Allocation.h b/Engine/Source/Runtime/D3D12RHI/Private/D3D12Allocation.h
--- a/Engine/Source/Runtime/D3D12RHI/Private/D3D12Allocation.h
+++ b/Engine/Source/Runtime/D3D12RHI/Private/D3D12Allocation.h
@@ -409,8 +409,15 @@ public:
 	~FD3D12DefaultBufferPool() { delete Allocator; }
 
 	bool SupportsAllocation(D3D12_HEAP_TYPE InHeapType, D3D12_RESOURCE_FLAGS InResourceFlags, EBufferUsageFlags InBufferUsage, ED3D12ResourceStateMode InResourceStateMode, uint32 Alignment) const;
-	void AllocDefaultResource(D3D12_HEAP_TYPE InHeapType, const D3D12_RESOURCE_DESC& InDesc, EBufferUsageFlags InBufferUsage, ED3D12ResourceStateMode InResourceStateMode,
-		D3D12_RESOURCE_STATES InCreateState, uint32 InAlignment, const TCHAR* InName, FD3D12ResourceLocation& ResourceLocation);
+	void AllocDefaultResource(
+		D3D12_HEAP_TYPE InHeapType,
+		const D3D12_RESOURCE_DESC& InDesc,
+		EBufferUsageFlags InBufferUsage,
+		ED3D12ResourceStateMode InResourceStateMode,
+		ED3D12Access InCreateD3D12Access,
+		uint32 InAlignment,
+		const TCHAR* InName,
+		FD3D12ResourceLocation& ResourceLocation);
 	void CleanUpAllocations(uint64 InFrameLag);
 	void UpdateMemoryStats(uint32& IOMemoryAllocated, uint32& IOMemoryUsed, uint32& IOMemoryFree, uint32& IOMemoryEndFree, uint32& IOAlignmentWaste, uint32& IOAllocatedPageCount, uint32& IOFullPageCount);
 
@@ -435,14 +442,23 @@ public:
 	FD3D12DefaultBufferAllocator(FD3D12Device* InParent, FRHIGPUMask VisibleNodes);
 
 	// Grab a buffer from the available buffers or create a new buffer if none are available
-	void AllocDefaultResource(D3D12_HEAP_TYPE InHeapType, const D3D12_RESOURCE_DESC& pDesc, EBufferUsageFlags InBufferUsage, ED3D12ResourceStateMode InResourceStateMode, D3D12_RESOURCE_STATES InCreateState, FD3D12ResourceLocation& ResourceLocation, uint32 Alignment, const TCHAR* Name);
+	void AllocDefaultResource(
+		D3D12_HEAP_TYPE InHeapType,
+		const D3D12_RESOURCE_DESC& pDesc,
+		EBufferUsageFlags InBufferUsage,
+		ED3D12ResourceStateMode InResourceStateMode,
+		ED3D12Access InCreateD3D12Access,
+		FD3D12ResourceLocation& ResourceLocation,
+		uint32 Alignment,
+		const TCHAR* Name);
+
 	void FreeDefaultBufferPools();
 	void BeginFrame(FD3D12ContextArray const& Contexts);
 	void CleanupFreeBlocks(uint64 InFrameLag);	
 	void UpdateMemoryStats();
 
 	static bool IsPlacedResource(D3D12_RESOURCE_FLAGS InResourceFlags, ED3D12ResourceStateMode InResourceStateMode, uint32 Alignment);
-	static D3D12_RESOURCE_STATES GetDefaultInitialResourceState(D3D12_HEAP_TYPE InHeapType, EBufferUsageFlags InBufferFlags, ED3D12ResourceStateMode InResourceStateMode);
+	static ED3D12Access GetDefaultInitialD3D12Access(D3D12_HEAP_TYPE InHeapType, EBufferUsageFlags InBufferFlags, ED3D12ResourceStateMode InResourceStateMode);
 
 private:
 
@@ -921,7 +937,14 @@ class FD3D12TextureAllocatorPool : public FD3D12DeviceChild, public FD3D12MultiN
 public:
 	FD3D12TextureAllocatorPool(FD3D12Device* Device, FRHIGPUMask VisibilityNode);
 
-	HRESULT AllocateTexture(FD3D12ResourceDesc Desc, const D3D12_CLEAR_VALUE* ClearValue, EPixelFormat UEFormat, FD3D12ResourceLocation& TextureLocation, const D3D12_RESOURCE_STATES InitialState, const TCHAR* Name);
+	HRESULT AllocateTexture(
+		FD3D12ResourceDesc Desc,
+		const D3D12_CLEAR_VALUE* ClearValue,
+		EPixelFormat UEFormat,
+		FD3D12ResourceLocation& TextureLocation,
+		ED3D12Access InInitialD3D12Access,
+		ED3D12Access InDefaultD3D12Access,
+		const TCHAR* Name);
 	
 	void BeginFrame(FD3D12ContextArray const& Contexts);
 	void CleanUpAllocations();
@@ -977,7 +1000,13 @@ public:
 
 	~FD3D12TextureAllocator();
 
-	HRESULT AllocateTexture(FD3D12ResourceDesc Desc, const D3D12_CLEAR_VALUE* ClearValue, FD3D12ResourceLocation& TextureLocation, const D3D12_RESOURCE_STATES InitialState, const TCHAR* Name);
+	HRESULT AllocateTexture(
+		FD3D12ResourceDesc Desc,
+		const D3D12_CLEAR_VALUE* ClearValue,
+		FD3D12ResourceLocation& TextureLocation,
+		ED3D12Access InInitialD3D12Access,
+		ED3D12Access InDefaultD3D12Access,
+		const TCHAR* Name);
 };
 
 class FD3D12TextureAllocatorPool : public FD3D12DeviceChild, public FD3D12MultiNodeGPUObject
@@ -985,7 +1014,14 @@ class FD3D12TextureAllocatorPool : public FD3D12DeviceChild, public FD3D12MultiN
 public:
 	FD3D12TextureAllocatorPool(FD3D12Device* Device, FRHIGPUMask VisibilityNode);
 
-	HRESULT AllocateTexture(FD3D12ResourceDesc Desc, const D3D12_CLEAR_VALUE* ClearValue, EPixelFormat UEFormat, FD3D12ResourceLocation& TextureLocation, const D3D12_RESOURCE_STATES InitialState, const TCHAR* Name);
+	HRESULT AllocateTexture(
+		FD3D12ResourceDesc Desc,
+		const D3D12_CLEAR_VALUE* ClearValue,
+		EPixelFormat UEFormat,
+		FD3D12ResourceLocation& TextureLocation,
+		ED3D12Access InInitialD3D12Access,
+		ED3D12Access InDefaultD3D12Access,
+		const TCHAR* Name);
 	
 	void BeginFrame(FD3D12ContextArray const& Contexts) {}
 	void CleanUpAllocations() { ReadOnlyTexturePool.CleanUpAllocations(0); }
diff --git a/Engine/Source/Runtime/D3D12RHI/Private/D3D12Buffer.cpp b/Engine/Source/Runtime/D3D12RHI/Private/D3D12Buffer.cpp
--- a/Engine/Source/Runtime/D3D12RHI/Private/D3D12Buffer.cpp
+++ b/Engine/Source/Runtime/D3D12RHI/Private/D3D12Buffer.cpp
@@ -27,16 +27,25 @@ FD3D12Buffer::~FD3D12Buffer()
 	}
 }
 
-void FD3D12Buffer::UploadResourceData(FD3D12CommandContext& CommandContext, D3D12_RESOURCE_STATES InDestinationState, FD3D12ResourceLocation& DestinationResourceLocation, const FD3D12ResourceLocation& SourceResourceLocation, uint32 Size)
+void FD3D12Buffer::UploadResourceData(
+	FD3D12CommandContext& CommandContext,
+	ED3D12Access InDestinationD3D12Access,
+	FD3D12ResourceLocation& DestinationResourceLocation,
+	const FD3D12ResourceLocation& SourceResourceLocation,
+	uint32 Size)
 {
 	FD3D12Resource* Destination = DestinationResourceLocation.GetResource();
+	check(Destination->GetDesc().Dimension == D3D12_RESOURCE_DIMENSION_BUFFER);
 
 	// Copy from the temporary upload heap to the default resource
 
+	const ERHIPipeline Pipe = CommandContext.GetPipeline();
+	const ED3D12Access InitialAccess = Destination->GetInitialAccess();
+
 	// if resource doesn't require state tracking then transition to copy dest here (could have been suballocated from shared resource) - not very optimal and should be batched
 	if (!Destination->RequiresResourceStateTracking())
 	{
-		CommandContext.AddTransitionBarrier(Destination, Destination->GetDefaultResourceState(), D3D12_RESOURCE_STATE_COPY_DEST, D3D12_RESOURCE_BARRIER_ALL_SUBRESOURCES);
+		CommandContext.AddBarrier(Destination, InitialAccess, ED3D12Access::CopyDest, D3D12_RESOURCE_BARRIER_ALL_SUBRESOURCES);
 	}
 
 	CommandContext.FlushResourceBarriers();
@@ -49,9 +58,9 @@ void FD3D12Buffer::UploadResourceData(FD3D12CommandContext& CommandContext, D3D1
 		Size);
 
 	// Update the resource state after the copy has been done (will take care of updating the residency as well)
-	if (InDestinationState != D3D12_RESOURCE_STATE_COPY_DEST)
+	if (InDestinationD3D12Access != ED3D12Access::CopyDest)
 	{
-		CommandContext.AddTransitionBarrier(Destination, D3D12_RESOURCE_STATE_COPY_DEST, InDestinationState, D3D12_RESOURCE_BARRIER_ALL_SUBRESOURCES);
+		CommandContext.AddBarrier(Destination, ED3D12Access::CopyDest, InDestinationD3D12Access, D3D12_RESOURCE_BARRIER_ALL_SUBRESOURCES);
 	}
 
 	CommandContext.UpdateResidency(SourceResourceLocation.GetResource());
@@ -59,7 +68,7 @@ void FD3D12Buffer::UploadResourceData(FD3D12CommandContext& CommandContext, D3D1
 	CommandContext.ConditionalSplitCommandList();
 
 	// If the resource is untracked, the destination state must match the default state of the resource.
-	check(Destination->RequiresResourceStateTracking() || (Destination->GetDefaultResourceState() == InDestinationState));
+	check(Destination->RequiresResourceStateTracking() || (Destination->GetDefaultAccess() == InDestinationD3D12Access));
 
 	// Buffer is now written and ready, so unlock the block (locked after creation and can be defragmented if needed)
 	DestinationResourceLocation.UnlockPoolData();
@@ -159,7 +168,7 @@ void FD3D12Adapter::AllocateBuffer(FD3D12Device* Device,
 	uint32 Size,
 	EBufferUsageFlags InUsage,
 	ED3D12ResourceStateMode InResourceStateMode,
-	D3D12_RESOURCE_STATES InCreateState,
+	ED3D12Access InCreateD3D12Access,
 	uint32 Alignment,
 	FD3D12Buffer* Buffer,
 	FD3D12ResourceLocation& ResourceLocation,
@@ -183,7 +192,7 @@ void FD3D12Adapter::AllocateBuffer(FD3D12Device* Device,
 	{
 		check(ResourceAllocator == nullptr);
 		check(InResourceStateMode != ED3D12ResourceStateMode::MultiState);
-		check(InCreateState == D3D12_RESOURCE_STATE_GENERIC_READ);
+		check(InCreateD3D12Access == ED3D12Access::GenericRead);
 		GetUploadHeapAllocator(Device->GetGPUIndex()).AllocUploadResource(Size, Alignment, ResourceLocation);
 		check(ResourceLocation.GetSize() >= Size);
 	}
@@ -191,23 +200,48 @@ void FD3D12Adapter::AllocateBuffer(FD3D12Device* Device,
 	{
 		if (ResourceAllocator)
 		{
-			ResourceAllocator->AllocateResource(Device->GetGPUIndex(), D3D12_HEAP_TYPE_DEFAULT, InDesc, InDesc.Width, Alignment, InResourceStateMode, InCreateState, nullptr, InDebugName, ResourceLocation);
+			ResourceAllocator->AllocateResource(
+				Device->GetGPUIndex(),
+				D3D12_HEAP_TYPE_DEFAULT,
+				InDesc,
+				InDesc.Width,
+				Alignment,
+				InCreateD3D12Access,
+				InResourceStateMode,
+				InCreateD3D12Access,
+				nullptr,
+				InDebugName,
+				ResourceLocation);
 		}
 		else
 		{
-			Device->GetDefaultBufferAllocator().AllocDefaultResource(D3D12_HEAP_TYPE_DEFAULT, InDesc, InUsage, InResourceStateMode, InCreateState, ResourceLocation, Alignment, InDebugName);
+			Device->GetDefaultBufferAllocator().AllocDefaultResource(
+				D3D12_HEAP_TYPE_DEFAULT,
+				InDesc,
+				InUsage,
+				InResourceStateMode,
+				InCreateD3D12Access,
+				ResourceLocation,
+				Alignment,
+				InDebugName);
 		}
 		ResourceLocation.SetOwner(Buffer);
 		check(ResourceLocation.GetSize() >= Size);
 	}
 }
 
+//
+// @TODO - This function is awkward and ambiguous because the CreateDesc includes an initial state
+// But it also takes a separate create state. Right now the value in CreateDesc is completely ignored
+// but that seems wrong... suggest removing InCreatED3D12Access argument in favor of using the value
+// in CreateDesc
+//
 FD3D12Buffer* FD3D12Adapter::CreateRHIBuffer(
 	const D3D12_RESOURCE_DESC& InDesc,
 	uint32 Alignment,
 	const FRHIBufferCreateDesc& CreateDesc,
 	ED3D12ResourceStateMode InResourceStateMode,
-	D3D12_RESOURCE_STATES InCreateState,
+	ED3D12Access InCreateD3D12Access,
 	bool bKeepUnlocked,
 	ID3D12ResourceAllocator* ResourceAllocator)
 {
@@ -240,7 +274,20 @@ FD3D12Buffer* FD3D12Adapter::CreateRHIBuffer(
 
 		if (!bDynamic || (Device->GetGPUIndex() == FirstGPUIndex) || EnumHasAnyFlags(CreateDesc.Usage, EBufferUsageFlags::MultiGPUAllocate))
 		{
-			AllocateBuffer(Device, InDesc, CreateDesc.Size, CreateDesc.Usage, InResourceStateMode, InCreateState, Alignment, NewBuffer, NewBuffer->ResourceLocation, ResourceAllocator, CreateDesc.DebugName, CreateDesc.OwnerName);
+			AllocateBuffer(
+				Device,
+				InDesc,
+				CreateDesc.Size,
+				CreateDesc.Usage,
+				InResourceStateMode,
+				InCreateD3D12Access,
+				Alignment,
+				NewBuffer,
+				NewBuffer->ResourceLocation,
+				ResourceAllocator,
+				CreateDesc.DebugName,
+				CreateDesc.OwnerName);
+
 			NewBuffer0 = NewBuffer;
 		}
 		else
@@ -352,6 +399,13 @@ void FD3D12Buffer::GetResourceDescAndAlignment(const FRHIBufferCreateDesc& Creat
 		ResourceDesc.Flags |= D3D12_RESOURCE_FLAG_ALLOW_SIMULTANEOUS_ACCESS;
 	}
 
+#if D3D12_RHI_RAYTRACING
+	if (EnumHasAnyFlags(CreateDesc.Usage, EBufferUsageFlags::AccelerationStructure))
+	{
+		ResourceDesc.Flags |= D3D12_RESOURCE_FLAG_RAYTRACING_ACCELERATION_STRUCTURE;
+	}
+#endif
+
 	if (EnumHasAnyFlags(CreateDesc.Usage, EBufferUsageFlags::ReservedResource))
 	{
 		checkf(CreateDesc.Stride <= GRHIGlobals.ReservedResources.TileSizeInBytes,
@@ -441,22 +495,31 @@ FD3D12DynamicRHI::FCreateBufferInternalResult FD3D12DynamicRHI::CreateBufferInte
 	const bool bSupportResourceStateTracking = !bIsDynamic && FD3D12DefaultBufferAllocator::IsPlacedResource(Desc.Flags, StateMode, Alignment) && Type.bWritable;
 
 	// Initial state is derived from the InResourceState if it supports tracking
-	const D3D12_RESOURCE_STATES DesiredState = bSupportResourceStateTracking ? Type.GetOptimalInitialState(CreateDesc.InitialState, false) :
-		FD3D12DefaultBufferAllocator::GetDefaultInitialResourceState(HeapType, CreateDesc.Usage, StateMode);
+	const ED3D12Access DesiredD3D12Access = bSupportResourceStateTracking ? 
+		Type.GetOptimalInitialD3D12Access(ConvertToD3D12Access(CreateDesc.InitialState), false) :
+		FD3D12DefaultBufferAllocator::GetDefaultInitialD3D12Access(HeapType, CreateDesc.Usage, StateMode);
 
 	// Setup the state at which the resource needs to be created - copy dest only supported for placed resources
-	const D3D12_RESOURCE_STATES CreateState = (bHasInitialData && bSupportResourceStateTracking) ? D3D12_RESOURCE_STATE_COPY_DEST : DesiredState;
-
-	FD3D12Buffer* Buffer = GetAdapter().CreateRHIBuffer(Desc, Alignment, CreateDesc, StateMode, CreateState, bHasInitialData, ResourceAllocator);
+	const ED3D12Access CreatED3D12Access =
+		(bHasInitialData && bSupportResourceStateTracking) ? ED3D12Access::CopyDest : DesiredD3D12Access;
+
+	FD3D12Buffer* Buffer = GetAdapter().CreateRHIBuffer(
+		Desc,
+		Alignment,
+		CreateDesc,
+		StateMode,
+		CreatED3D12Access,
+		bHasInitialData,
+		ResourceAllocator);
 	check(Buffer->ResourceLocation.IsValid());
 
-	return FCreateBufferInternalResult{ Buffer, DesiredState };
+	return FCreateBufferInternalResult{ Buffer, DesiredD3D12Access };
 }
 
 static FRHIBufferInitializer CreateBufferInitializerForWriting(FRHICommandListBase& RHICmdList, const FD3D12DynamicRHI::FCreateBufferInternalResult& CreateResult, const FRHIBufferCreateDesc& CreateDesc)
 {
 	FD3D12Buffer* Buffer = CreateResult.Buffer;
-	D3D12_RESOURCE_STATES DesiredState = CreateResult.DesiredState;
+	ED3D12Access DesiredD3D12Access = CreateResult.DesiredD3D12Access;
 
 	if (EnumHasAnyFlags(CreateDesc.Usage, EBufferUsageFlags::AnyDynamic))
 	{
@@ -471,10 +534,10 @@ static FRHIBufferInitializer CreateBufferInitializerForWriting(FRHICommandListBa
 	void* WritableData = UploadLocation.GetMappedBaseAddress();
 
 	return UE::RHICore::FCustomBufferInitializer(RHICmdList, Buffer, WritableData, CreateDesc.Size,
-		[Buffer = TRefCountPtr<FD3D12Buffer>(Buffer), DesiredState, UploadLocation = MoveTemp(UploadLocation)](FRHICommandListBase& RHICmdList) mutable
+		[Buffer = TRefCountPtr<FD3D12Buffer>(Buffer), DesiredD3D12Access, UploadLocation = MoveTemp(UploadLocation)](FRHICommandListBase& RHICmdList) mutable
 		{
 			RHICmdList.EnqueueLambda(
-				[Buffer = Buffer.GetReference(), DesiredState, UploadLocation = MoveTemp(UploadLocation)](FRHICommandListBase& ExecutingCmdList)
+				[Buffer = Buffer.GetReference(), DesiredD3D12Access, UploadLocation = MoveTemp(UploadLocation)](FRHICommandListBase& ExecutingCmdList)
 				{
 					const FRHIGPUMask EffectiveMask = ExecutingCmdList.GetGPUMask();
 					for (uint32 GPUIndex : EffectiveMask)
@@ -482,7 +545,12 @@ static FRHIBufferInitializer CreateBufferInitializerForWriting(FRHICommandListBa
 						FD3D12CommandContext& CommandContext = FD3D12CommandContext::Get(ExecutingCmdList, GPUIndex);
 						FD3D12Buffer* DeviceBuffer = CommandContext.RetrieveObject<FD3D12Buffer>(Buffer, GPUIndex);
 
-						Buffer->UploadResourceData(CommandContext, DesiredState, DeviceBuffer->ResourceLocation, UploadLocation, Buffer->GetSize());
+						Buffer->UploadResourceData(
+							CommandContext,
+							DesiredD3D12Access,
+							DeviceBuffer->ResourceLocation,
+							UploadLocation,
+							Buffer->GetSize());
 					}
 				});
 
@@ -650,7 +718,12 @@ void* FD3D12DynamicRHI::LockBuffer(FRHICommandListBase& RHICmdList, FD3D12Buffer
 				FD3D12CommandContext& Context = FD3D12CommandContext::Get(ExecutingCmdList, Node.GetFirstIndex());
 				uint64 SubAllocOffset = Buffer->ResourceLocation.GetOffsetFromBaseOfResource();
 
-				FScopedResourceBarrier ScopeResourceBarrierSource(Context, pResource, &Buffer->ResourceLocation, D3D12_RESOURCE_STATE_COPY_SOURCE, 0);
+				FScopedResourceBarrier ScopeResourceBarrierSource(
+					Context,
+					pResource,
+					ED3D12Access::CopySrc,
+					0);
+
 				// Don't need to transition upload heaps
 				Context.FlushResourceBarriers(); // Must flush so the desired state is actually set.
 
@@ -708,6 +781,9 @@ void FD3D12DynamicRHI::UnlockBuffer(FRHICommandListBase& RHICmdList, FD3D12Buffe
 		check(Buffer->IsHeadLink());
 
 		RHICmdList.EnqueueLambda([
+#if EXECUTE_DEBUG_COMMAND_LISTS
+			this,
+#endif
 			RootBuffer = Buffer,
 			LockedData = MoveTemp(LockedData)
 		](FRHICommandListBase& ExecutingCmdList)
@@ -725,7 +801,12 @@ void FD3D12DynamicRHI::UnlockBuffer(FRHICommandListBase& RHICmdList, FD3D12Buffe
 				// Clear the resource if still bound to make sure the SRVs are rebound again on next operation (and get correct resource transitions enqueued)
 				Context.ConditionalClearShaderResource(&Buffer.ResourceLocation, EShaderParameterTypeMask::SRVMask);
 
-				FScopedResourceBarrier ScopeResourceBarrierDest(Context, DestResource, &Buffer.ResourceLocation, D3D12_RESOURCE_STATE_COPY_DEST, 0);
+				FScopedResourceBarrier ScopeResourceBarrierDest(
+					Context,
+					DestResource,
+					ED3D12Access::CopyDest,
+					0);
+
 				// Don't need to transition upload heaps
 				Context.FlushResourceBarriers();
 
@@ -889,8 +970,16 @@ void FD3D12CommandContext::RHICopyBufferRegion(FRHIBuffer* DestBufferRHI, uint64
 	check(DstOffset + NumBytes <= DestBufferDesc.Width);
 	check(SrcOffset + NumBytes <= SourceBufferDesc.Width);
 
-	FScopedResourceBarrier ScopeResourceBarrierSrc(*this, pSourceResource, &SourceBuffer->ResourceLocation, D3D12_RESOURCE_STATE_COPY_SOURCE, 0);
-	FScopedResourceBarrier ScopeResourceBarrierDst(*this, pDestResource  , &DestBuffer->ResourceLocation, D3D12_RESOURCE_STATE_COPY_DEST  , 0);
+	FScopedResourceBarrier ScopeResourceBarrierSrc(
+		*this,
+		pSourceResource,
+		ED3D12Access::CopySrc,
+		0);
+	FScopedResourceBarrier ScopeResourceBarrierDst(
+		*this,
+		pDestResource,
+		ED3D12Access::CopyDest,
+		0);
 	FlushResourceBarriers();
 
 	CopyBufferRegionChecked(
@@ -900,9 +989,6 @@ void FD3D12CommandContext::RHICopyBufferRegion(FRHIBuffer* DestBufferRHI, uint64
 		SourceBuffer->ResourceLocation.GetOffsetFromBaseOfResource() + SrcOffset,
 		NumBytes
 	);
-	
-	UpdateResidency(pDestResource);
-	UpdateResidency(pSourceResource);
 
 	ConditionalSplitCommandList();
 
diff --git a/Engine/Source/Runtime/D3D12RHI/Private/D3D12CommandContext.cpp b/Engine/Source/Runtime/D3D12RHI/Private/D3D12CommandContext.cpp
--- a/Engine/Source/Runtime/D3D12RHI/Private/D3D12CommandContext.cpp
+++ b/Engine/Source/Runtime/D3D12RHI/Private/D3D12CommandContext.cpp
@@ -51,56 +51,12 @@ FD3D12CommandContextBase::FD3D12CommandContextBase(class FD3D12Adapter* InParent
 {
 }
 
-static D3D12_RESOURCE_STATES GetValidResourceStates(ED3D12QueueType CommandListType)
-{
-	// For reasons, we can't just list the allowed states, we have to list the disallowed states.
-	// For reference on allowed/disallowed states, see:
-	//    https://microsoft.github.io/DirectX-Specs/d3d/CPUEfficiency.html#state-support-by-command-list-type
-
-	const D3D12_RESOURCE_STATES DisallowedDirectStates =
-		static_cast<D3D12_RESOURCE_STATES>(0);
-
-	const D3D12_RESOURCE_STATES DisallowedComputeStates =
-		DisallowedDirectStates |
-		D3D12_RESOURCE_STATE_VERTEX_AND_CONSTANT_BUFFER |
-		D3D12_RESOURCE_STATE_INDEX_BUFFER |
-		D3D12_RESOURCE_STATE_RENDER_TARGET |
-		D3D12_RESOURCE_STATE_DEPTH_WRITE |
-		D3D12_RESOURCE_STATE_DEPTH_READ |
-		D3D12_RESOURCE_STATE_PIXEL_SHADER_RESOURCE |
-		D3D12_RESOURCE_STATE_STREAM_OUT |
-		D3D12_RESOURCE_STATE_INDIRECT_ARGUMENT |
-		D3D12_RESOURCE_STATE_RESOLVE_DEST |
-		D3D12_RESOURCE_STATE_RESOLVE_SOURCE;
-
-	const D3D12_RESOURCE_STATES DisallowedCopyStates =
-		DisallowedComputeStates |
-		D3D12_RESOURCE_STATE_UNORDERED_ACCESS |
-		D3D12_RESOURCE_STATE_NON_PIXEL_SHADER_RESOURCE;
-
-
-	if (CommandListType == ED3D12QueueType::Copy)
-	{
-		return ~DisallowedCopyStates;
-	}
-
-	if (CommandListType == ED3D12QueueType::Async)
-	{
-		return ~DisallowedComputeStates;
-	}
-
-	return ~DisallowedDirectStates;
-}
-
-D3D12_RESOURCE_STATES FD3D12CommandContext::SkipFastClearEliminateState = D3D12_RESOURCE_STATES(0);
-
 FD3D12CommandContext::FD3D12CommandContext(FD3D12Device* InParent, ED3D12QueueType QueueType, bool InIsDefaultContext)
 	: FD3D12ContextCommon(InParent, QueueType, InIsDefaultContext)
 	, FD3D12CommandContextBase(InParent->GetParentAdapter(), InParent->GetGPUMask())
 	, FD3D12DeviceChild(InParent)
 	, ConstantsAllocator(InParent, InParent->GetGPUMask())
 	, StateCache(*this, InParent->GetGPUMask())
-	, ValidResourceStates(GetValidResourceStates(QueueType))
 	, StageConstantBuffers{
 		FD3D12ConstantBuffer(InParent, ConstantsAllocator),
 		FD3D12ConstantBuffer(InParent, ConstantsAllocator),
@@ -110,7 +66,6 @@ FD3D12CommandContext::FD3D12CommandContext(FD3D12Device* InParent, ED3D12QueueTy
 		FD3D12ConstantBuffer(InParent, ConstantsAllocator),
 	}
 {
-	SkipFastClearEliminateState = D3D12_RESOURCE_STATES(0);
 	StaticUniformBuffers.AddZeroed(FUniformBufferStaticSlotRegistry::Get().GetSlotCount());
 	ClearState();
 }
@@ -289,6 +244,7 @@ FD3D12ContextCommon::FD3D12ContextCommon(FD3D12Device* Device, ED3D12QueueType Q
 	, TimestampQueries(Device, QueueType, D3D12_QUERY_TYPE_TIMESTAMP)
 	, OcclusionQueries(Device, QueueType, D3D12_QUERY_TYPE_OCCLUSION)
 	, PipelineStatsQueries(Device, QueueType, D3D12_QUERY_TYPE_PIPELINE_STATISTICS)
+	, Barriers(Device->GetParentAdapter()->CreateBarriersForContext())
 {
 }
 
@@ -501,6 +457,22 @@ void FD3D12ContextCommon::Finalize(TArray<FD3D12Payload*>& OutPayloads)
 	OutPayloads.Append(MoveTemp(Payloads));
 }
 
+ERHIPipeline FD3D12ContextCommon::GetRHIPipeline() const
+{
+	switch (QueueType)
+	{
+	case ED3D12QueueType::Direct:
+		return ERHIPipeline::Graphics;
+	case ED3D12QueueType::Async:
+		return ERHIPipeline::AsyncCompute;
+	default:
+		checkNoEntry();
+		return ERHIPipeline::None;
+	}
+	static_assert(static_cast<uint32>(ERHIPipeline::Num) == 2, 
+		"Above logic only handles two types of pipes");
+}
+
 uint32 FD3D12CommandContext::GetFrameFenceCounter() const
 {
 	return GetParentDevice()->GetParentAdapter()->GetFrameFence().GetNextFenceToSignal();
@@ -629,6 +601,13 @@ void FD3D12ContextCommon::NewPayload()
 
 void FD3D12ContextCommon::FlushCommands(ED3D12FlushFlags FlushFlags)
 {
+#if D3D12RHI_IDLE_AFTER_EVERY_GPU_EVENT
+	if (!IsDefaultContext())
+	{
+		return;
+	}
+#endif
+
 	// We should only be flushing the default context
 	check(IsDefaultContext());
 
@@ -803,76 +782,12 @@ IRHIComputeContext* FD3D12DynamicRHI::RHIGetCommandContext(ERHIPipeline Pipeline
 
 void FD3D12DynamicRHI::RHICreateTransition(FRHITransition* Transition, const FRHITransitionCreateInfo& CreateInfo)
 {
-	// Construct the data in-place on the transition instance
-	FD3D12TransitionData* Data = new (Transition->GetPrivateData<FD3D12TransitionData>()) FD3D12TransitionData;
-
-	Data->SrcPipelines = CreateInfo.SrcPipelines;
-	Data->DstPipelines = CreateInfo.DstPipelines;
-	Data->CreateFlags = CreateInfo.Flags;
-
-	const bool bCrossPipeline = (CreateInfo.SrcPipelines != CreateInfo.DstPipelines) && (!EnumHasAnyFlags(Data->CreateFlags, ERHITransitionCreateFlags::NoFence));
-	const bool bAsyncToAllPipelines = ((CreateInfo.SrcPipelines == ERHIPipeline::AsyncCompute) && (CreateInfo.DstPipelines == ERHIPipeline::All));
-
-	Data->bCrossPipeline = bCrossPipeline;
-
-	// In DX12 we cannot perform resource barrier with graphics state on the AsyncCompute pipe
-	// This check is here to be able to force a crosspipe transition coming from AsyncCompute with graphics states to be split and processed in the both the Async and Graphics pipe
-	// This case can be removed when using EB on DX12 
-	if (bAsyncToAllPipelines)
-	{
-		for (const FRHITransitionInfo& TransitionInfo : CreateInfo.TransitionInfos)
-		{
-			if (EnumHasAnyFlags(TransitionInfo.AccessAfter, ERHIAccess::SRVGraphics))
-			{
-				Data->bAsyncToAllPipelines = true;
-				Data->bCrossPipeline = false;
-				break;
-			}
-		}
-	}
-
-	if ((Data->bCrossPipeline) || (Data->bAsyncToAllPipelines))
-	{
-		// Create one sync point per device, per source pipe
-		for (uint32 Index : FRHIGPUMask::All())
-		{
-			TRHIPipelineArray<FD3D12SyncPointRef>& DeviceSyncPoints = Data->SyncPoints.Emplace_GetRef();
-			for (ERHIPipeline Pipeline : MakeFlagsRange(CreateInfo.SrcPipelines))
-			{
-				DeviceSyncPoints[Pipeline] = FD3D12SyncPoint::Create(ED3D12SyncPointType::GPUOnly, TEXT("Transition"));
-			}
-		}
-	}
-
-	Data->TransitionInfos = CreateInfo.TransitionInfos;
-	Data->AliasingInfos = CreateInfo.AliasingInfos;
-
-	uint32 AliasingOverlapCount = 0;
-
-	for (const FRHITransientAliasingInfo& AliasingInfo : Data->AliasingInfos)
-	{
-		AliasingOverlapCount += AliasingInfo.Overlaps.Num();
-	}
-
-	Data->AliasingOverlaps.Reserve(AliasingOverlapCount);
-
-	for (FRHITransientAliasingInfo& AliasingInfo : Data->AliasingInfos)
-	{
-		const int32 OverlapCount = AliasingInfo.Overlaps.Num();
-
-		if (OverlapCount > 0)
-		{
-			const int32 OverlapOffset = Data->AliasingOverlaps.Num();
-			Data->AliasingOverlaps.Append(AliasingInfo.Overlaps.GetData(), OverlapCount);
-			AliasingInfo.Overlaps = MakeArrayView(&Data->AliasingOverlaps[OverlapOffset], OverlapCount);
-		}
-	}
+	GetAdapter().CreateTransition(Transition, CreateInfo);
 }
 
 void FD3D12DynamicRHI::RHIReleaseTransition(FRHITransition* Transition)
 {
-	// Destruct the transition data
-	Transition->GetPrivateData<FD3D12TransitionData>()->~FD3D12TransitionData();
+	GetAdapter().ReleaseTransition(Transition);
 }
 
 IRHITransientResourceAllocator* FD3D12DynamicRHI::RHICreateTransientResourceAllocator()
diff --git a/Engine/Source/Runtime/D3D12RHI/Private/D3D12CommandContext.h b/Engine/Source/Runtime/D3D12RHI/Private/D3D12CommandContext.h
--- a/Engine/Source/Runtime/D3D12RHI/Private/D3D12CommandContext.h
+++ b/Engine/Source/Runtime/D3D12RHI/Private/D3D12CommandContext.h
@@ -15,6 +15,7 @@ D3D12CommandContext.h: D3D12 Command Context Interfaces
 #include "D3D12StateCachePrivate.h"
 #include "D3D12Submission.h"
 #include "D3D12Texture.h"
+#include COMPILED_PLATFORM_HEADER(D3D12BarriersFactory.h)
 #include "Experimental/Containers/RobinHoodHashTable.h"
 
 #include "RHICoreShader.h"
@@ -253,13 +254,14 @@ public:
 	// The type of command lists this context records.
 	ED3D12QueueType const QueueType;
 	bool IsAsyncComputeContext() const { return QueueType == ED3D12QueueType::Async; }
+	ERHIPipeline GetRHIPipeline() const;
 
 	// True for the immediate context (@todo remove this)
 	bool const bIsDefaultContext;
 	bool IsDefaultContext() const { return bIsDefaultContext; }
 
 	bool IsOpen() const { return CommandList != nullptr; }
-	bool IsPendingCommands() const { return IsOpen() || ResourceBarrierBatcher.Num(); }
+	bool IsPendingCommands() const { return IsOpen() || Barriers->GetNumPendingBarriers(); }
 
 	FD3D12SyncPoint* GetContextSyncPoint()
 	{
@@ -282,15 +284,17 @@ public:
 
 	void BindDiagnosticBuffer(FD3D12RootSignature const* RootSignature, ED3D12PipelineType PipelineType);
 
+	FD3D12QueryAllocator& GetTimestampQueries()
+	{
+		return TimestampQueries;
+	}
+
 private:
 	// Allocators to manage query heaps
 	FD3D12QueryAllocator TimestampQueries;
 	FD3D12QueryAllocator OcclusionQueries;
 	FD3D12QueryAllocator PipelineStatsQueries;
 
-	// Batches resource barriers together until it's explicitly flushed
-	FD3D12ResourceBarrierBatcher ResourceBarrierBatcher;
-
 	// The active D3D12 command list where recorded D3D commands are directed.
 	// This is swapped when command lists are split (e.g. when signalling a fence).
 	FD3D12CommandList* CommandList = nullptr;
@@ -312,7 +316,6 @@ public:
 		// Nothing to do
 	}
 
-protected:
 	// Returns the current command list (or creates a new one if the command list was not open).
 	FD3D12CommandList& GetCommandList()
 	{
@@ -321,6 +324,7 @@ protected:
 		return *CommandList;
 	}
 
+protected:
 	enum class EMarkerType { In, Out };
 	void WriteMarker(D3D12_GPU_VIRTUAL_ADDRESS Address, uint32 Value, EMarkerType Type);
 
@@ -347,6 +351,8 @@ protected:
 
 	uint32 ActiveQueries = 0;
 
+	const TUniquePtr<FD3D12BarriersFactory::BarriersForContextType> Barriers;
+
 public:
 	// Open the command list if it's not already open.
 	void OpenIfNotAlready()
@@ -418,18 +424,18 @@ public:
 #endif
 
 	// Resource transition / barrier functions. These get batched and recorded into the command list when FlushResourceBarriers() is called.
-	void AddTransitionBarrier(FD3D12Resource* pResource, D3D12_RESOURCE_STATES Before, D3D12_RESOURCE_STATES After, uint32 Subresource);
-	void AddAliasingBarrier(ID3D12Resource* InResourceBefore, ID3D12Resource* InResourceAfter);
-	void AddUAVBarrier();
+	void AddGlobalBarrier(
+		ED3D12Access InD3D12AccessBefore,
+		ED3D12Access InD3D12AccessAfter);
+
+	void AddBarrier(
+		const FD3D12Resource* pResource,
+		ED3D12Access InD3D12AccessBefore,
+		ED3D12Access InD3D12AccessAfter,
+		uint32 Subresource);
 
 	// Flushes the batched resource barriers to the current command list
 	void FlushResourceBarriers();
-
-	// Functions for transitioning a resource. The Before and After state cannot be D3D12_RESOURCE_STATE_TBD, needs be known at the time of the call
-	void TransitionResource(FD3D12Resource* Resource, D3D12_RESOURCE_STATES Before, D3D12_RESOURCE_STATES After, uint32 Subresource);
-	
-private:
-	void TransitionResource(FD3D12Resource* InResource, uint32 InSubresourceIndex, D3D12_RESOURCE_STATES InBeforeState, D3D12_RESOURCE_STATES InAfterState);
 };
 
 //
@@ -577,9 +583,6 @@ public:
 	D3D12_SHADING_RATE				VRSShadingRate = D3D12_SHADING_RATE_1X1;
 #endif
 
-	static D3D12_RESOURCE_STATES SkipFastClearEliminateState;
-	D3D12_RESOURCE_STATES ValidResourceStates;
-
 	/** Constant buffers for Set*ShaderParameter calls. */
 	FD3D12ConstantBuffer StageConstantBuffers[SF_NumStandardFrequencies];
 
@@ -789,6 +792,11 @@ public:
 		return RetrieveTexture(Texture, GetGPUIndex());
 	}
 
+	FORCEINLINE_DEBUGGABLE const FD3D12Texture* RetrieveTexture(const FRHITexture* Texture)
+	{
+		return RetrieveTexture(const_cast<FRHITexture*>(Texture), GetGPUIndex());
+	}
+
 	uint32 GetFrameFenceCounter() const;
 
 	uint32 GetGPUIndex() const { return GPUMask.ToIndex(); }
@@ -854,12 +862,6 @@ private:
 	void DispatchWorkGraphShaderBundle(FRHIShaderBundle* ShaderBundle, FRHIBuffer* RecordArgBuffer, TConstArrayView<FRHIShaderParameterResource> SharedBindlessParameters, TConstArrayView<FRHIShaderBundleComputeDispatch> Dispatches);
 	void DispatchWorkGraphShaderBundle(FRHIShaderBundle* ShaderBundle, FRHIBuffer* RecordArgBuffer, const FRHIShaderBundleGraphicsState& BundleState, TConstArrayView<FRHIShaderParameterResource> SharedBindlessParameters, TConstArrayView<FRHIShaderBundleGraphicsDispatch> Dispatches);
 
-	void HandleDiscardResources          (TArrayView<const FRHITransition*> Transitions, bool bIsBeginTransition);
-	void HandleResourceTransitions       (const struct FD3D12TransitionData* TransitionData, bool& bUAVBarrier);
-	void HandleTransientAliasing         (const struct FD3D12TransitionData* TransitionData);
-	void HandleResourceDiscardTransitions(const struct FD3D12TransitionData* TransitionData, TArray<struct FD3D12DiscardResource>& ResourcesToDiscard);
-	void HandleReservedResourceCommits   (const struct FD3D12TransitionData* TransitionData);
-
 	TArray<FRHIUniformBuffer*> StaticUniformBuffers;
 	const FRHIShaderBindingLayout* ShaderBindinglayout = nullptr;
 
@@ -1187,21 +1189,6 @@ private:
 	TStaticArray<FD3D12CommandContext*, MAX_NUM_GPUS> PhysicalContexts;
 };
 
-struct FD3D12TransitionData
-{
-	ERHIPipeline SrcPipelines, DstPipelines;
-	ERHITransitionCreateFlags CreateFlags = ERHITransitionCreateFlags::None;
-
-	TArray<FRHITransitionInfo, TInlineAllocator<4, FConcurrentLinearArrayAllocator>> TransitionInfos;
-	TArray<FRHITransientAliasingInfo, TInlineAllocator<4, FConcurrentLinearArrayAllocator>> AliasingInfos;
-	TArray<FRHITransientAliasingOverlap, TInlineAllocator<4, FConcurrentLinearArrayAllocator>> AliasingOverlaps;
-
-	TArray<TRHIPipelineArray<FD3D12SyncPointRef>, TInlineAllocator<MAX_NUM_GPUS>> SyncPoints;
-
-	bool bCrossPipeline = false;
-	bool bAsyncToAllPipelines = false;
-};
-
 class FD3D12ContextArray : public TRHIPipelineArray<FD3D12CommandContextBase*>
 {
 public:
diff --git a/Engine/Source/Runtime/D3D12RHI/Private/D3D12CommandList.cpp b/Engine/Source/Runtime/D3D12RHI/Private/D3D12CommandList.cpp
--- a/Engine/Source/Runtime/D3D12RHI/Private/D3D12CommandList.cpp
+++ b/Engine/Source/Runtime/D3D12RHI/Private/D3D12CommandList.cpp
@@ -4,18 +4,6 @@
 #include "D3D12RHIPrivate.h"
 #include "RHIValidation.h"
 
-static int32 GD3D12BatchResourceBarriers = 1;
-static FAutoConsoleVariableRef CVarD3D12BatchResourceBarriers(
-	TEXT("d3d12.BatchResourceBarriers"),
-	GD3D12BatchResourceBarriers,
-	TEXT("Whether to allow batching resource barriers"));
-
-static int32 GD3D12ExtraDepthTransitions = 0;
-static FAutoConsoleVariableRef CVarD3D12ExtraDepthTransitions(
-	TEXT("d3d12.ExtraDepthTransitions"),
-	GD3D12ExtraDepthTransitions,
-	TEXT("Adds extra transitions for the depth buffer to fix validation issues. However, this currently breaks async compute"));
-
 void FD3D12CommandList::UpdateResidency(const FD3D12Resource* Resource)
 {
 #if ENABLE_RESIDENCY_MANAGEMENT
@@ -61,53 +49,27 @@ void FD3D12CommandList::AddToResidencySet(TConstArrayView<FD3D12ResidencyHandle*
 }
 #endif // ENABLE_RESIDENCY_MANAGEMENT
 
-
-void FD3D12ContextCommon::AddTransitionBarrier(FD3D12Resource* pResource, D3D12_RESOURCE_STATES Before, D3D12_RESOURCE_STATES After, uint32 Subresource)
+void FD3D12ContextCommon::AddGlobalBarrier(
+	ED3D12Access InD3D12AccessBefore,
+	ED3D12Access InD3D12AccessAfter)
 {
-	if (Before != After)
-	{
-		ResourceBarrierBatcher.AddTransition(pResource, Before, After, Subresource);
-
-		UpdateResidency(pResource);
-
-		if (!GD3D12BatchResourceBarriers)
-		{
-			FlushResourceBarriers();
-		}
-	}
-	else
-	{
-		ensureMsgf(0, TEXT("AddTransitionBarrier: Before == After (%d)"), (uint32)Before);
-	}
+	Barriers->AddGlobalBarrier(*this, InD3D12AccessBefore, InD3D12AccessAfter);
 }
 
-void FD3D12ContextCommon::AddUAVBarrier()
+void FD3D12ContextCommon::AddBarrier(
+	const FD3D12Resource* pResource,
+	ED3D12Access InD3D12AccessBefore,
+	ED3D12Access InD3D12AccessAfter,
+	uint32 Subresource)
 {
-	ResourceBarrierBatcher.AddUAV();
-
-	if (!GD3D12BatchResourceBarriers)
-	{
-		FlushResourceBarriers();
-	}
+	check(pResource);
+	Barriers->AddBarrier(*this, pResource, InD3D12AccessBefore, InD3D12AccessAfter, Subresource);
+	UpdateResidency(pResource);
 }
 
-void FD3D12ContextCommon::AddAliasingBarrier(ID3D12Resource* InResourceBefore, ID3D12Resource* InResourceAfter)
-{
-	ResourceBarrierBatcher.AddAliasingBarrier(InResourceBefore, InResourceAfter);
-
-	if (!GD3D12BatchResourceBarriers)
-	{
-		FlushResourceBarriers();
-	}
-}
-
-
 void FD3D12ContextCommon::FlushResourceBarriers()
 {
-	if (ResourceBarrierBatcher.Num())
-	{
-		ResourceBarrierBatcher.FlushIntoCommandList(GetCommandList(), TimestampQueries);
-	}
+	Barriers->FlushIntoCommandList(GetCommandList(), GetTimestampQueries());
 }
 
 FD3D12CommandAllocator::FD3D12CommandAllocator(FD3D12Device* Device, ED3D12QueueType QueueType)
@@ -450,58 +412,6 @@ FD3D12CommandList::FState::FState(FD3D12CommandAllocator* CommandAllocator, FD3D
 	}
 }
 
-void FD3D12ContextCommon::TransitionResource(FD3D12Resource* Resource, D3D12_RESOURCE_STATES Before, D3D12_RESOURCE_STATES After, uint32 Subresource)
-{
-	check(Resource);
-	check(Resource->RequiresResourceStateTracking());
-	check(!((After & (D3D12_RESOURCE_STATE_PIXEL_SHADER_RESOURCE | D3D12_RESOURCE_STATE_NON_PIXEL_SHADER_RESOURCE)) && (Resource->GetDesc().Flags & D3D12_RESOURCE_FLAG_DENY_SHADER_RESOURCE)));
-	check(Before != D3D12_RESOURCE_STATE_TBD);
-	check(After != D3D12_RESOURCE_STATE_TBD);
-
-#ifdef PLATFORM_SUPPORTS_RESOURCE_COMPRESSION
-	After |= Resource->GetCompressedState();
-#endif
-
-#if ENABLE_RHI_VALIDATION
-	FString IncompatibilityReason;
-	if (CheckResourceStateCompatibility(After, Resource->GetDesc().Flags, IncompatibilityReason) == false)
-	{
-		UE_LOG(LogRHI, Error, TEXT("Incompatible Transition State for Resource %s - %s"), *Resource->GetName().ToString(), *IncompatibilityReason);
-	}
-#endif
-
-	UpdateResidency(Resource);
-
-	TransitionResource(Resource, Subresource, Before, After);
-}
-
-static inline bool IsTransitionNeeded(D3D12_RESOURCE_STATES Before, D3D12_RESOURCE_STATES After, FD3D12Resource* InResource = nullptr)
-{
-	check(Before != D3D12_RESOURCE_STATE_CORRUPT && After != D3D12_RESOURCE_STATE_CORRUPT);
-	check(Before != D3D12_RESOURCE_STATE_TBD && After != D3D12_RESOURCE_STATE_TBD);
-
-	// COMMON is an oddball state that doesn't follow the RESOURE_STATE pattern of 
-	// having exactly one bit set so we need to special case these
-	if (After == D3D12_RESOURCE_STATE_COMMON)
-	{
-		// Before state should not have the common state otherwise it's invalid transition
-		check(Before != D3D12_RESOURCE_STATE_COMMON);
-		return true;
-	}
-
-	return Before != After;
-}
-
-void FD3D12ContextCommon::TransitionResource(FD3D12Resource* InResource, uint32 InSubresourceIndex, D3D12_RESOURCE_STATES InBeforeState, D3D12_RESOURCE_STATES InAfterState)
-{
-	check(InBeforeState != D3D12_RESOURCE_STATE_TBD);
-	check(InAfterState != D3D12_RESOURCE_STATE_TBD);
-
-	if (IsTransitionNeeded(InBeforeState, InAfterState, InResource))
-	{
-		AddTransitionBarrier(InResource, InBeforeState, InAfterState, InSubresourceIndex);
-	}
-}
 
 namespace D3D12RHI
 {
diff --git a/Engine/Source/Runtime/D3D12RHI/Private/D3D12CommandList.h b/Engine/Source/Runtime/D3D12RHI/Private/D3D12CommandList.h
--- a/Engine/Source/Runtime/D3D12RHI/Private/D3D12CommandList.h
+++ b/Engine/Source/Runtime/D3D12RHI/Private/D3D12CommandList.h
@@ -17,7 +17,6 @@ class FD3D12Device;
 class FD3D12DynamicRHI;
 class FD3D12QueryAllocator;
 class FD3D12Queue;
-class FD3D12ResourceBarrierBatcher;
 
 struct FD3D12ResidencyHandle;
 
@@ -67,7 +66,6 @@ private:
 	friend FD3D12ContextCommon;
 	friend FD3D12DynamicRHI;
 	friend FD3D12Queue;
-	friend FD3D12ResourceBarrierBatcher;
 
 	FD3D12CommandList(FD3D12CommandList const&) = delete;
 	FD3D12CommandList(FD3D12CommandList&&) = delete;
diff --git a/Engine/Source/Runtime/D3D12RHI/Private/D3D12Commands.cpp b/Engine/Source/Runtime/D3D12RHI/Private/D3D12Commands.cpp
--- a/Engine/Source/Runtime/D3D12RHI/Private/D3D12Commands.cpp
+++ b/Engine/Source/Runtime/D3D12RHI/Private/D3D12Commands.cpp
@@ -12,13 +12,6 @@ D3D12Commands.cpp: D3D RHI commands implementation.
 #include "RHIUniformBufferUtilities.h"
 #include "RHICoreTransitions.h"
 
-static int32 GD3D12AllowDiscardResources = 1;
-static FAutoConsoleVariableRef CVarD3D12AllowDiscardResources(
-	TEXT("d3d12.AllowDiscardResources"),
-	GD3D12AllowDiscardResources,
-	TEXT("Whether to call DiscardResources after transient aliasing acquire. This is not needed on some platforms if newly acquired resources are cleared before use."),
-	ECVF_RenderThreadSafe);
-
 using namespace D3D12RHI;
 
 inline void ValidateBoundShader(FD3D12StateCache& InStateCache, FRHIShader* InShaderRHI)
@@ -132,6 +125,12 @@ FD3D12ResourceLocation& FD3D12CommandContext::SetupIndirectArgument(FRHIBuffer*
 
 void FD3D12CommandContext::PostGpuEvent()
 {
+#if D3D12RHI_IDLE_AFTER_EVERY_GPU_EVENT
+	AddGlobalBarrier(
+		ED3D12Access::Common,
+		ED3D12Access::Common);
+#endif
+
 	UnsetExplicitDescriptorCache();
 
 	ConditionalSplitCommandList();
@@ -174,603 +173,14 @@ void FD3D12CommandContext::RHIDispatchIndirectComputeShader(FRHIBuffer* Argument
 	PostGpuEvent();
 }
 
-template <typename FunctionType>
-void EnumerateSubresources(FD3D12Resource* Resource, const FRHITransitionInfo& Info, FD3D12Texture* Texture, FunctionType Function)
-{
-	uint32 FirstMipSlice = 0;
-	uint32 FirstArraySlice = 0;
-	uint32 FirstPlaneSlice = 0;
-
-	uint32 MipCount = Resource->GetMipLevels();
-	uint32 ArraySize = Resource->GetArraySize();
-	uint32 PlaneCount = Resource->GetPlaneCount();
-
-	uint32 IterationMipCount = MipCount;
-	uint32 IterationArraySize = ArraySize;
-	uint32 IterationPlaneCount = PlaneCount;
-
-	if (!Info.IsAllMips())
-	{
-		FirstMipSlice = Info.MipIndex;
-		IterationMipCount = 1;
-	}
-
-	if (!Info.IsAllArraySlices())
-	{
-		FirstArraySlice = Info.ArraySlice;
-		IterationArraySize = 1;
-	}
-
-	if (!Info.IsAllPlaneSlices())
-	{
-		FirstPlaneSlice = Info.PlaneSlice;
-		IterationPlaneCount = 1;
-	}
-
-	for (uint32 PlaneSlice = FirstPlaneSlice; PlaneSlice < FirstPlaneSlice + IterationPlaneCount; ++PlaneSlice)
-	{
-		for (uint32 ArraySlice = FirstArraySlice; ArraySlice < FirstArraySlice + IterationArraySize; ++ArraySlice)
-		{
-			for (uint32 MipSlice = FirstMipSlice; MipSlice < FirstMipSlice + IterationMipCount; ++MipSlice)
-			{
-				const uint32 Subresource = D3D12CalcSubresource(MipSlice, ArraySlice, PlaneSlice, MipCount, ArraySize);
-				const FD3D12RenderTargetView* RTV = nullptr;
-#if PLATFORM_REQUIRES_TYPELESS_RESOURCE_DISCARD_WORKAROUND
-				if (Texture)
-				{
-					RTV = Texture->GetRenderTargetView(MipSlice, ArraySlice);
-				}
-#endif // PLATFORM_REQUIRES_TYPELESS_RESOURCE_DISCARD_WORKAROUND
-				Function(Subresource, RTV);
-			}
-		}
-	}
-}
-
-static
-std::tuple<FD3D12Resource*,FD3D12Texture*> GetResourceAndTexture(FD3D12CommandContext& Context, const FRHITransitionInfo& Info)
-{
-	switch (Info.Type)
-	{
-	case FRHITransitionInfo::EType::UAV:
-	{
-		FD3D12UnorderedAccessView* UAV = Context.RetrieveObject<FD3D12UnorderedAccessView_RHI>(Info.UAV);
-		check(UAV);
-		if (UAV)
-		{
-			return { UAV->GetResource(), nullptr };
-		}
-		return { nullptr, nullptr };
-	}
-	case FRHITransitionInfo::EType::Buffer:
-	{
-		// Resource may be null if this is a multi-GPU resource not present on the current GPU
-		FD3D12Buffer* Buffer = Context.RetrieveObject<FD3D12Buffer>(Info.Buffer);
-
-		if (Buffer)
-		{
-			return { Buffer->GetResource(), nullptr };
-		}
-		return { nullptr, nullptr };
-	}
-	case FRHITransitionInfo::EType::Texture:
-	{
-		// Resource may be null if this is a multi-GPU resource not present on the current GPU
-		FD3D12Texture* Texture = Context.RetrieveTexture(Info.Texture);
-
-		if (Texture)
-		{
-			return { Texture->GetResource(), Texture };
-		}
-		return { nullptr, nullptr };
-	}
-	case FRHITransitionInfo::EType::BVH:
-	{
-		// Nothing special required for BVH transitions - handled inside d3d12 raytracing directly via UAV barriers and don't need explicit state changes
-		return { nullptr, nullptr };
-	}
-	default:
-		checkNoEntry();
-		return { nullptr, nullptr };
-		break;
-	}
-}
-
-template <typename FunctionType>
-void ProcessResource(FD3D12CommandContext& Context, const FRHITransitionInfo& Info, FunctionType Function)
-{
-	auto [Resource, Texture] = GetResourceAndTexture(Context, Info);
-	FD3D12Texture* DiscardTextureOut = nullptr;
-
-	if ((Info.Type == FRHITransitionInfo::EType::Texture) && Texture)
-	{
-#if PLATFORM_REQUIRES_TYPELESS_RESOURCE_DISCARD_WORKAROUND
-		if (Texture->GetRequiresTypelessResourceDiscardWorkaround())
-		{
-			DiscardTextureOut = Texture;
-		}
-#endif // #if PLATFORM_REQUIRES_TYPELESS_RESOURCE_DISCARD_WORKAROUND
-	}
-
-	if (Resource)
-	{
-		Function(Info, Resource, Texture, DiscardTextureOut);
-	}
-}
-
-// Pipe changes which are not ending with graphics or targeting all pipelines are handle during begin
-static bool ProcessTransitionDuringBegin(const FD3D12TransitionData* Data)
-{
-	// Source pipelines aren't on all pipelines
-	const bool bSrcPipelinesNotAll = !EnumHasAllFlags(Data->SrcPipelines, ERHIPipeline::All);
-
-	// Source and destination pipelines are different
-	const bool bSrcDstPipelinesDiffer = Data->SrcPipelines != Data->DstPipelines;
-
-	// Destination pipeline is not only graphics
-	const bool bDstPipelineNotGraphics = Data->DstPipelines != ERHIPipeline::Graphics;
-
-	// Destination pipelines include all pipelines
-	const bool bDstPipelinesIncludeAll = EnumHasAllFlags(Data->DstPipelines, ERHIPipeline::All);
-
-	return bSrcPipelinesNotAll && ( (bSrcDstPipelinesDiffer && bDstPipelineNotGraphics) || bDstPipelinesIncludeAll );
-}
-
-static bool ShouldProcessTransition(const FD3D12TransitionData* Data, bool bIsBeginTransition, ERHIPipeline ExecutingPipeline)
-{
-	// Special DX12 case where crosspipe transitions from AsyncCompute with graphics state can only be processed on the ERHIPipeline::Graphics pipe
-	if (Data->bAsyncToAllPipelines)
-	{
-		if ((bIsBeginTransition == false) && (ExecutingPipeline == ERHIPipeline::Graphics))
-		{
-			return true;
-		}
-		
-		if ((bIsBeginTransition) && (ExecutingPipeline == ERHIPipeline::AsyncCompute))
-		{
-			return true;
-		}
-
-		return false;
-	}
-
-	return (bIsBeginTransition ? ProcessTransitionDuringBegin(Data) : !ProcessTransitionDuringBegin(Data));
-}
-
-struct FD3D12DiscardResource
-{
-	FD3D12DiscardResource(FD3D12Resource* InResource, EResourceTransitionFlags InFlags, uint32 InSubresource = D3D12_RESOURCE_BARRIER_ALL_SUBRESOURCES, const FD3D12Texture* InTexture = nullptr, const FD3D12RenderTargetView* InRTV = nullptr)
-		: Resource(InResource)
-		, Flags(InFlags)
-		, Subresource(InSubresource)
-#if PLATFORM_REQUIRES_TYPELESS_RESOURCE_DISCARD_WORKAROUND
-		, Texture(InTexture)
-		, RTV(InRTV)
-#endif // #if PLATFORM_REQUIRES_TYPELESS_RESOURCE_DISCARD_WORKAROUND
-	{}
-
-	FD3D12Resource* Resource;
-	EResourceTransitionFlags Flags;
-	uint32 Subresource;
-#if PLATFORM_REQUIRES_TYPELESS_RESOURCE_DISCARD_WORKAROUND
-	const FD3D12Texture* Texture = nullptr;
-	const FD3D12RenderTargetView* RTV = nullptr;
-#endif // #if PLATFORM_REQUIRES_TYPELESS_RESOURCE_DISCARD_WORKAROUND
-};
-
-void FD3D12CommandContext::HandleReservedResourceCommits(const struct FD3D12TransitionData* TransitionData)
-{
-	for (const FRHITransitionInfo& Info : TransitionData->TransitionInfos)
-	{
-		if (const FRHICommitResourceInfo* CommitInfo = Info.CommitInfo.GetPtrOrNull())
-		{
-			if (Info.Type == FRHITransitionInfo::EType::Buffer)
-			{
-				FD3D12Buffer* Buffer = RetrieveObject<FD3D12Buffer>(Info.Buffer);
-				SetReservedBufferCommitSize(Buffer, CommitInfo->SizeInBytes);
-			}
-			else
-			{
-				checkNoEntry();
-			}
-		}
-	}
-}
-
-bool IsImpossibleAsyncDiscardTransition(ERHIPipeline Pipeline, FRHITexture* Texture)
-{
-	return Pipeline == ERHIPipeline::AsyncCompute && Texture && EnumHasAnyFlags(Texture->GetDesc().Flags, ETextureCreateFlags::RenderTargetable | ETextureCreateFlags::DepthStencilTargetable);
-}
-
-void FD3D12CommandContext::HandleResourceDiscardTransitions(
-	const FD3D12TransitionData* TransitionData,
-	TArray<FD3D12DiscardResource>& ResourcesToDiscard)
-{
-	for (const FRHITransitionInfo& Info : TransitionData->TransitionInfos)
-	{
-		const UE::RHICore::FResourceState ResourceState(*this, TransitionData->SrcPipelines, TransitionData->DstPipelines, Info);
-
-		if (!EnumHasAnyFlags(ResourceState.AccessBefore, ERHIAccess::Discard))
-		{
-			continue;
-		}
-
-		ProcessResource(*this, Info, [&](const FRHITransitionInfo& Info, FD3D12Resource* Resource, FD3D12Texture* Texture, FD3D12Texture* DiscardTexture)
-		{
-			D3D12_RESOURCE_STATES StateAfter  = GetInitialResourceState(Resource->GetDesc(), QueueType);
-			D3D12_RESOURCE_STATES StateBefore = StateAfter;
-
-			if (ResourceState.AccessBefore != ERHIAccess::Discard)
-			{
-				StateBefore = GetD3D12ResourceState(ResourceState.AccessBefore & ~ERHIAccess::Discard, QueueType, Texture);
-			}
-
-			const bool bTransition = StateBefore != StateAfter;
-
-			if (bTransition)
-			{
-				// Transitions here should only occur on the Direct queue and when the prior Discard operation failed due to being on async compute.
-				ensure(IsImpossibleAsyncDiscardTransition(ResourceState.SrcPipelines, Texture) && QueueType == ED3D12QueueType::Direct);
-			}
-
-			if (Info.IsWholeResource() || Resource->GetSubresourceCount() == 1)
-			{
-				if (bTransition)
-				{
-					TransitionResource(Resource, StateBefore, StateAfter, D3D12_RESOURCE_BARRIER_ALL_SUBRESOURCES);
-				}
-				else
-				{
-					UpdateResidency(Resource);
-				}
-
-				const FD3D12RenderTargetView* RTV = nullptr;
-#if PLATFORM_REQUIRES_TYPELESS_RESOURCE_DISCARD_WORKAROUND
-				if (DiscardTexture)
-				{
-					RTV = DiscardTexture->GetRenderTargetView(0, -1);
-				}
-#endif // #if PLATFORM_REQUIRES_TYPELESS_RESOURCE_DISCARD_WORKAROUND
-				ResourcesToDiscard.Emplace(Resource, Info.Flags, D3D12_RESOURCE_BARRIER_ALL_SUBRESOURCES, Texture, RTV);
-			}
-			else
-			{
-				EnumerateSubresources(Resource, Info, DiscardTexture, [&](uint32 Subresource, const FD3D12RenderTargetView* RTV)
-				{
-					if (bTransition)
-					{
-						TransitionResource(Resource, StateBefore, StateAfter, Subresource);
-					}
-					else
-					{
-						UpdateResidency(Resource);
-					}
-
-					ResourcesToDiscard.Emplace(Resource, Info.Flags, Subresource, DiscardTexture, RTV);
-				});
-			}
-		});
-	}
-}
-
-void FD3D12CommandContext::HandleDiscardResources(
-	TArrayView<const FRHITransition*> Transitions,
-	bool bIsBeginTransition)
-{
-	TArray<FD3D12DiscardResource> ResourcesToDiscard;
-
-	for (const FRHITransition* Transition : Transitions)
-	{
-		const FD3D12TransitionData* Data = Transition->GetPrivateData<FD3D12TransitionData>();
-
-		if (ProcessTransitionDuringBegin(Data) == bIsBeginTransition)
-		{
-			HandleResourceDiscardTransitions(Data, ResourcesToDiscard);
-		}
-	}
-
-	if (!GD3D12AllowDiscardResources)
-	{
-		return;
-	}
-
-	if (!ResourcesToDiscard.IsEmpty())
-	{
-		FlushResourceBarriers();
-	}
-
-	for (const FD3D12DiscardResource& DiscardResource : ResourcesToDiscard)
-	{
-#if PLATFORM_REQUIRES_TYPELESS_RESOURCE_DISCARD_WORKAROUND
-		if (DiscardResource.Texture && DiscardResource.RTV && DiscardResource.Texture->GetRequiresTypelessResourceDiscardWorkaround())
-		{
-			FLinearColor ClearColor = DiscardResource.Texture->GetClearColor();
-			GraphicsCommandList()->ClearRenderTargetView(DiscardResource.RTV->GetOfflineCpuHandle(), reinterpret_cast<float*>(&ClearColor), 0, nullptr);
-			UpdateResidency(DiscardResource.RTV->GetResource());
-		}
-		else
-#endif // #if PLATFORM_REQUIRES_TYPELESS_RESOURCE_DISCARD_WORKAROUND
-		{
-			if (DiscardResource.Subresource == D3D12_RESOURCE_BARRIER_ALL_SUBRESOURCES)
-			{
-				GraphicsCommandList()->DiscardResource(DiscardResource.Resource->GetResource(), nullptr);
-			}
-			else
-			{
-				D3D12_DISCARD_REGION Region;
-				Region.NumRects = 0;
-				Region.pRects = nullptr;
-				Region.FirstSubresource = DiscardResource.Subresource;
-				Region.NumSubresources = 1;
-
-				GraphicsCommandList()->DiscardResource(DiscardResource.Resource->GetResource(), &Region);
-			}
-		}
-	}
-}
-
-void FD3D12CommandContext::HandleTransientAliasing(const FD3D12TransitionData* TransitionData)
-{
-	for (const FRHITransientAliasingInfo& Info : TransitionData->AliasingInfos)
-	{
-		FD3D12BaseShaderResource* BaseShaderResource = nullptr;
-		switch (Info.Type)
-		{
-		case FRHITransientAliasingInfo::EType::Buffer:
-		{
-			// Resource may be null if this is a multi-GPU resource not present on the current GPU
-			FD3D12Buffer* Buffer = RetrieveObject<FD3D12Buffer>(Info.Buffer);
-			check(Buffer || GNumExplicitGPUsForRendering > 1);
-			BaseShaderResource = Buffer;
-			break;
-		}
-		case FRHITransientAliasingInfo::EType::Texture:
-		{
-			// Resource may be null if this is a multi-GPU resource not present on the current GPU
-			FD3D12Texture* Texture = RetrieveTexture(Info.Texture);
-			check(Texture || GNumExplicitGPUsForRendering > 1);
-			BaseShaderResource = Texture;
-			break;
-		}
-		default:
-			checkNoEntry();
-			break;
-		}
-
-		// Resource may be null if this is a multi-GPU resource not present on the current GPU
-		if (!BaseShaderResource)
-		{
-			continue;
-		}
-
-		FD3D12Resource* Resource = BaseShaderResource->ResourceLocation.GetResource();
-		if (Info.Action == FRHITransientAliasingInfo::EAction::Acquire)
-		{
-			TRACE_CPUPROFILER_EVENT_SCOPE(D3D12RHI::AcquireTransient);
-			AddAliasingBarrier(nullptr, Resource->GetResource());
-		}
-	}
-}
-
-void FD3D12CommandContext::HandleResourceTransitions(const FD3D12TransitionData* TransitionData, bool& bUAVBarrier)
-{
-	for (const FRHITransitionInfo& Info : TransitionData->TransitionInfos)
-	{
-		if (!Info.Resource)
-		{
-			continue;
-		}
-
-		UE::RHICore::FResourceState ResourceState(*this, TransitionData->SrcPipelines, TransitionData->DstPipelines, Info);
-
-		bUAVBarrier |=
-			EnumHasAnyFlags(ResourceState.AccessBefore, ERHIAccess::UAVMask) &&
-			EnumHasAnyFlags(ResourceState.AccessAfter,  ERHIAccess::UAVMask);
-
-		// Skip duplicate transitions. This happens most frequently with implicit ones from NeedsExtraTransitions.
-		if (ResourceState.AccessBefore == ResourceState.AccessAfter)
-		{
-			continue;
-		}
-
-		// Very specific case that needs to be removed with EB
-		// a UAV -> SRVMAsk on the AsyncPipe get split in two: UAV->SRVCompute on Async and SRVCompute->SRVMask on Gfx
-		// On the Async pipe is going to be : UAV->SRVMask(that is automatically converted in UAV->SRVCompute)
-		// On the Direct(Gfx) pipe instead needs to be SRVCompute->SRVMask therefore the check there to change the Before state only on the Direct pipe.
-		if (TransitionData->bAsyncToAllPipelines && (ResourceState.AccessAfter == ERHIAccess::SRVMask) && (QueueType == ED3D12QueueType::Direct))
-		{
-			ResourceState.AccessBefore = ERHIAccess::SRVCompute;
-		}
-
-		// Process transitions which are forced during begin because those contain transition from Graphics to Compute and should
-		// help remove forced patch up command lists for async compute to run on the graphics queue
-		ProcessResource(*this, Info, [&](const FRHITransitionInfo& Info, FD3D12Resource* Resource, FD3D12Texture* Texture, FD3D12Texture* DiscardTexture)
-		{
-			if (!Resource->RequiresResourceStateTracking())
-			{
-				return;
-			}
-
-			if (ResourceState.AccessAfter == ERHIAccess::Discard && IsImpossibleAsyncDiscardTransition(ResourceState.DstPipelines, Texture))
-			{
-				return;
-			}
-
-			D3D12_RESOURCE_STATES StateBefore;
-
-			if (EnumHasAnyFlags(ResourceState.AccessBefore, ERHIAccess::Discard))
-			{
-				StateBefore = GetInitialResourceState(Resource->GetDesc(), QueueType);
-			}
-			else
-			{
-				StateBefore = GetD3D12ResourceState(ResourceState.AccessBefore, QueueType, Texture);
-			}
-
-			if (ResourceState.AccessBefore != ERHIAccess::Present)
-			{
-				check(StateBefore != D3D12_RESOURCE_STATE_COMMON);
-			}
-
-			D3D12_RESOURCE_STATES StateAfter = ResourceState.AccessAfter == ERHIAccess::Discard ? GetInitialResourceState(Resource->GetDesc(), QueueType) : GetD3D12ResourceState(ResourceState.AccessAfter & ~ERHIAccess::Discard, QueueType, Texture);
-
-			// enqueue the correct transitions
-			if (Info.IsWholeResource() || Resource->GetSubresourceCount() == 1)
-			{
-				TransitionResource(Resource, StateBefore, StateAfter, D3D12_RESOURCE_BARRIER_ALL_SUBRESOURCES);
-			}
-			else
-			{
-				// high level rendering is controlling transition ranges, at this level this is an index not a range
-				check(Info.MipIndex != FRHISubresourceRange::kAllSubresources);
-				check(Info.ArraySlice != FRHISubresourceRange::kAllSubresources);
-				check(Info.PlaneSlice != FRHISubresourceRange::kAllSubresources);
-				const uint32 Subresource = D3D12CalcSubresource(Info.MipIndex, Info.ArraySlice, Info.PlaneSlice, Resource->GetMipLevels(), Resource->GetArraySize());
-				check(Subresource < Resource->GetSubresourceCount());
-				TransitionResource(Resource, StateBefore, StateAfter, Subresource);
-			}
-		});
-	}
-}
-
 void FD3D12CommandContext::RHIBeginTransitions(TArrayView<const FRHITransition*> Transitions)
 {
-	const ERHIPipeline CurrentPipeline = GetPipeline();
-	const bool bIsBeginTransition = true;
-
-	for (const FRHITransition* Transition : Transitions)
-	{
-		const FD3D12TransitionData* Data = Transition->GetPrivateData<FD3D12TransitionData>();
-
-		if (ShouldProcessTransition ( Data, bIsBeginTransition, CurrentPipeline))
-		{
-			HandleTransientAliasing(Data);
-		}
-	}
-
-	HandleDiscardResources(Transitions, bIsBeginTransition );
-
-	bool bUAVBarrier = false;
-
-	for (const FRHITransition* Transition : Transitions)
-	{
-		const FD3D12TransitionData* Data = Transition->GetPrivateData<FD3D12TransitionData>();
-
-		// Handle transition during BeginTransitions?
-		if (ShouldProcessTransition(Data, bIsBeginTransition, CurrentPipeline))
-		{
-			HandleResourceTransitions(Data, bUAVBarrier);
-		}
-	}
-
-	if (bUAVBarrier)
-	{
-		StateCache.FlushComputeShaderCache(true);
-	}
-
-	// Signal fences
-	for (const FRHITransition* Transition : Transitions)
-	{
-		const auto* Data = Transition->GetPrivateData<FD3D12TransitionData>();
-		if (Data->bCrossPipeline)
-		{
-			const TRHIPipelineArray<FD3D12SyncPointRef>& DeviceSyncPoints = Data->SyncPoints[GetGPUIndex()];
-
-			if (DeviceSyncPoints[CurrentPipeline])
-			{
-				SignalSyncPoint(DeviceSyncPoints[CurrentPipeline]);
-			}
-		}
-	}
+	Barriers->BeginTransitions(*this, Transitions);
 }
 
 void FD3D12CommandContext::RHIEndTransitions(TArrayView<const FRHITransition*> Transitions)
 {
-	const ERHIPipeline CurrentPipeline = GetPipeline();
-	const bool bIsBeginTransition = false;
-
-	// Wait for fences
-	{
-		for (const FRHITransition* Transition : Transitions)
-		{
-			const auto* Data = Transition->GetPrivateData<FD3D12TransitionData>();
-
-			if (Data->bAsyncToAllPipelines)
-			{
-				const TRHIPipelineArray<FD3D12SyncPointRef>& DeviceSyncPoints = Data->SyncPoints[GetGPUIndex()];
-				if (CurrentPipeline == ERHIPipeline::Graphics)
-				{
-					WaitSyncPoint(DeviceSyncPoints[ERHIPipeline::AsyncCompute]);
-				}
-			}
-			else if (Data->bCrossPipeline)
-			{
-				const TRHIPipelineArray<FD3D12SyncPointRef>& DeviceSyncPoints = Data->SyncPoints[GetGPUIndex()];
-
-				for (ERHIPipeline SrcPipeline : MakeFlagsRange(Data->SrcPipelines))
-				{
-					if (SrcPipeline != CurrentPipeline && DeviceSyncPoints[SrcPipeline])
-					{
-						WaitSyncPoint(DeviceSyncPoints[SrcPipeline]);
-					}
-				}
-			}
-		}
-	}
-
-	// Update reserved resource memory mapping
-	for (const FRHITransition* Transition : Transitions)
-	{
-		const FD3D12TransitionData* Data = Transition->GetPrivateData<FD3D12TransitionData>();
-
-		HandleReservedResourceCommits(Data);
-	}
-
-	for (const FRHITransition* Transition : Transitions)
-	{
-		const FD3D12TransitionData* Data = Transition->GetPrivateData<FD3D12TransitionData>();
-
-		if (ShouldProcessTransition(Data, bIsBeginTransition, CurrentPipeline))
-		{
-			HandleTransientAliasing(Data);
-		}
-	}
-
-	HandleDiscardResources(Transitions, false /** bIsBeginTransitions */);
-
-	bool bUAVBarrier = false;
-
-	for (const FRHITransition* Transition : Transitions)
-	{
-		const FD3D12TransitionData* Data = Transition->GetPrivateData<FD3D12TransitionData>();
-
-		// Handle transition during EndTransitions?
-		if (ShouldProcessTransition(Data, bIsBeginTransition, CurrentPipeline))
-		{
-			HandleResourceTransitions(Data, bUAVBarrier);
-		}
-	}
-
-	if (bUAVBarrier)
-	{
-		StateCache.FlushComputeShaderCache(true);
-	}
-
-	// Signal fences
-	for (const FRHITransition* Transition : Transitions)
-	{
-		const auto* Data = Transition->GetPrivateData<FD3D12TransitionData>();
-
-		if ((Data->bAsyncToAllPipelines) && (CurrentPipeline == ERHIPipeline::AsyncCompute))
-		{
-			const TRHIPipelineArray<FD3D12SyncPointRef>& DeviceSyncPoints = Data->SyncPoints[GetGPUIndex()];
-
-			if (DeviceSyncPoints[CurrentPipeline])
-			{
-				SignalSyncPoint(DeviceSyncPoints[CurrentPipeline]);
-			}
-		}
-	}
+	Barriers->EndTransitions(*this, Transitions);
 }
 
 void FD3D12CommandContext::RHISetStaticUniformBuffers(const FUniformBufferStaticBindings& InUniformBuffers)
@@ -829,7 +239,15 @@ void FD3D12CommandContext::RHICopyToStagingBuffer(FRHIBuffer* SourceBufferRHI, F
 		// Unknown aligment requirement for sub allocated read back buffer data
 		uint32 AllocationAlignment = 16;
 		const D3D12_RESOURCE_DESC BufferDesc = CD3DX12_RESOURCE_DESC::Buffer(NumBytes, D3D12_RESOURCE_FLAG_NONE);		
-		GetParentDevice()->GetDefaultBufferAllocator().AllocDefaultResource(D3D12_HEAP_TYPE_READBACK, BufferDesc, BUF_None, ED3D12ResourceStateMode::SingleState, D3D12_RESOURCE_STATE_COPY_DEST, StagingBuffer->ResourceLocation, AllocationAlignment, TEXT("StagedRead"));
+		GetParentDevice()->GetDefaultBufferAllocator().AllocDefaultResource(
+			D3D12_HEAP_TYPE_READBACK,
+			BufferDesc,
+			BUF_None,
+			ED3D12ResourceStateMode::SingleState,
+			ED3D12Access::CopyDest,
+			StagingBuffer->ResourceLocation,
+			AllocationAlignment,
+			TEXT("StagedRead"));
 		check(StagingBuffer->ResourceLocation.GetSize() == NumBytes);
 		StagingBuffer->ShadowBufferSize = NumBytes;
 	}
@@ -849,7 +267,7 @@ void FD3D12CommandContext::RHICopyToStagingBuffer(FRHIBuffer* SourceBufferRHI, F
 
 #if D3D12_RHI_RAYTRACING
 		// TODO: Buffer in D3D12_RESOURCE_STATE_RAYTRACING_ACCELERATION_STRUCTURE needs CopyRaytracingAccelerationStructure
-		bool bRayTracingAccellerationStruct = ((pSourceResource->RequiresResourceStateTracking() == false) && (pSourceResource->GetDefaultResourceState() == D3D12_RESOURCE_STATE_RAYTRACING_ACCELERATION_STRUCTURE));
+		bool bRayTracingAccellerationStruct = ((pSourceResource->RequiresResourceStateTracking() == false) && (pSourceResource->GetDesc().Flags & D3D12_RESOURCE_FLAG_RAYTRACING_ACCELERATION_STRUCTURE));
 #else
 		bool bRayTracingAccellerationStruct = false;
 #endif
diff --git a/Engine/Source/Runtime/D3D12RHI/Private/D3D12DescriptorCache.cpp b/Engine/Source/Runtime/D3D12RHI/Private/D3D12DescriptorCache.cpp
--- a/Engine/Source/Runtime/D3D12RHI/Private/D3D12DescriptorCache.cpp
+++ b/Engine/Source/Runtime/D3D12RHI/Private/D3D12DescriptorCache.cpp
@@ -452,8 +452,6 @@ D3D12_GPU_DESCRIPTOR_HANDLE FD3D12DescriptorCache::BuildSRVTable(EShaderFrequenc
 	D3D12_CPU_DESCRIPTOR_HANDLE DestDescriptor = CurrentViewHeap->GetCPUSlotHandle(FirstSlotIndex);
 	D3D12_CPU_DESCRIPTOR_HANDLE SrcDescriptors[MAX_SRVS];
 
-	const D3D12_RESOURCE_STATES ValidResourceStates = Context.ValidResourceStates;
-
 	for (uint32 SlotIndex = 0; SlotIndex < SlotsNeeded; SlotIndex++)
 	{
 		if (FD3D12ShaderResourceView* SRV = SRVs[SlotIndex])
@@ -513,8 +511,6 @@ void FD3D12DescriptorCache::SetSRVTable(EShaderFrequency ShaderStage, const FD3D
 
 void FD3D12DescriptorCache::PrepareBindlessViews(EShaderFrequency ShaderStage, TConstArrayView<FD3D12ShaderResourceView*> SRVs, TConstArrayView<FD3D12UnorderedAccessView*> UAVs)
 {
-	const D3D12_RESOURCE_STATES ValidResourceStates = Context.ValidResourceStates;
-
 	for (FD3D12ShaderResourceView* SRV : SRVs)
 	{
 		if (ensure(SRV))
diff --git a/Engine/Source/Runtime/D3D12RHI/Private/D3D12Device.cpp b/Engine/Source/Runtime/D3D12RHI/Private/D3D12Device.cpp
--- a/Engine/Source/Runtime/D3D12RHI/Private/D3D12Device.cpp
+++ b/Engine/Source/Runtime/D3D12RHI/Private/D3D12Device.cpp
@@ -424,8 +424,7 @@ void FD3D12Device::SetupAfterDeviceCreation()
 			256,
 			CreateDesc,
 			ED3D12ResourceStateMode::MultiState,
-			D3D12_RESOURCE_STATE_INDIRECT_ARGUMENT
-		);
+			ED3D12Access::IndirectArgs);
 	}
 #endif // D3D12_RHI_RAYTRACING
 
