D3D12 Enhanced Barriers

diff --git a/Engine/Source/Runtime/D3D12RHI/Private/D3D12PoolAllocator.cpp b/Engine/Source/Runtime/D3D12RHI/Private/D3D12PoolAllocator.cpp
--- a/Engine/Source/Runtime/D3D12RHI/Private/D3D12PoolAllocator.cpp
+++ b/Engine/Source/Runtime/D3D12RHI/Private/D3D12PoolAllocator.cpp
@@ -110,7 +110,17 @@ void FD3D12MemoryPool::Init()
 			LLM_SCOPE_BYTAG(D3D12AllocatorUnused);
 #endif
 			const D3D12_HEAP_PROPERTIES HeapProps = CD3DX12_HEAP_PROPERTIES(InitConfig.HeapType, GetGPUMask().GetNative(), GetVisibilityMask().GetNative());
-			VERIFYD3D12RESULT(Adapter->CreateBuffer(HeapProps, GetGPUMask(), InitConfig.InitialResourceState, ED3D12ResourceStateMode::SingleState, InitConfig.InitialResourceState, PoolSize, BackingResource.GetInitReference(), TEXT("Resource PoolAllocator Underlying Buffer"), InitConfig.ResourceFlags));
+			VERIFYD3D12RESULT(Adapter->CreateBuffer(
+				HeapProps,
+				GetGPUMask(),
+				InitConfig.InitialD3D12Access,
+				ED3D12ResourceStateMode::SingleState,
+				InitConfig.InitialD3D12Access,
+				PoolSize,
+				BackingResource.GetInitReference(),
+				TEXT("Resource PoolAllocator Underlying Buffer"),
+				InitConfig.ResourceFlags));
+
 #if UE_MEMORY_TRACE_ENABLED
 			MemoryTrace_MarkAllocAsHeap(BackingResource->GetGPUVirtualAddress(), TraceHeapId);
 #endif
@@ -124,7 +134,8 @@ void FD3D12MemoryPool::Init()
 
 #if D3D12_RHI_RAYTRACING
 	// Elevates the raytracing acceleration structure heap priority, which may help performance / stability in low memory conditions.
-	if (IsGPUOnly(InitConfig.HeapType) && InitConfig.InitialResourceState == D3D12_RESOURCE_STATE_RAYTRACING_ACCELERATION_STRUCTURE)
+	if (IsGPUOnly(InitConfig.HeapType)
+		&& EnumHasAnyFlags(InitConfig.InitialD3D12Access, ED3D12Access::BVHRead | ED3D12Access::BVHWrite))
 	{
 		for (uint32 GPUIndex : GetGPUMask())
 		{
@@ -182,23 +193,23 @@ FD3D12ResourceInitConfig FD3D12PoolAllocator::GetResourceAllocatorInitConfig(D3D
 	{
 		// should only have this flag and no other flags
 		check(InBufferUsage == BUF_AccelerationStructure);
-		InitConfig.InitialResourceState = D3D12_RESOURCE_STATE_RAYTRACING_ACCELERATION_STRUCTURE;
+		InitConfig.InitialD3D12Access = ED3D12Access::BVHRead | ED3D12Access::BVHWrite;
 	}
 	else
 #endif // D3D12_RHI_RAYTRACING
-		if (InitConfig.HeapType == D3D12_HEAP_TYPE_READBACK)
-		{
-			InitConfig.InitialResourceState = D3D12_RESOURCE_STATE_COPY_DEST;
-		}
-		else if (EnumHasAnyFlags(InBufferUsage, BUF_UnorderedAccess))
-		{
-			check(InResourceFlags & D3D12_RESOURCE_FLAG_ALLOW_UNORDERED_ACCESS);
-			InitConfig.InitialResourceState = D3D12_RESOURCE_STATE_UNORDERED_ACCESS;
-		}
-		else
-		{
-			InitConfig.InitialResourceState = D3D12_RESOURCE_STATE_GENERIC_READ;
-		}
+	if (InitConfig.HeapType == D3D12_HEAP_TYPE_READBACK)
+	{
+		InitConfig.InitialD3D12Access = ED3D12Access::CopyDest;
+	}
+	else if (EnumHasAnyFlags(InBufferUsage, BUF_UnorderedAccess))
+	{
+		check(InResourceFlags & D3D12_RESOURCE_FLAG_ALLOW_UNORDERED_ACCESS);
+		InitConfig.InitialD3D12Access = ED3D12Access::UAVMask;
+	}
+	else
+	{
+		InitConfig.InitialD3D12Access = ED3D12Access::GenericRead;
+	}
 
 	InitConfig.HeapFlags = D3D12_HEAP_FLAG_ALLOW_ONLY_BUFFERS;
 	if (EnumHasAnyFlags(InBufferUsage, BUF_DrawIndirect))
@@ -274,39 +285,67 @@ bool FD3D12PoolAllocator::SupportsAllocation(D3D12_HEAP_TYPE InHeapType, D3D12_R
 }
 
 
-void FD3D12PoolAllocator::AllocDefaultResource(D3D12_HEAP_TYPE InHeapType, const FD3D12ResourceDesc& InDesc, EBufferUsageFlags InBufferUsage, ED3D12ResourceStateMode InResourceStateMode,
-	D3D12_RESOURCE_STATES InCreateState, uint32 InAllocationAlignment, const TCHAR* InName, FD3D12ResourceLocation& ResourceLocation)
+void FD3D12PoolAllocator::AllocDefaultResource(
+	D3D12_HEAP_TYPE InHeapType,
+	const FD3D12ResourceDesc& InDesc,
+	EBufferUsageFlags InBufferUsage,
+	ED3D12ResourceStateMode InResourceStateMode,
+	ED3D12Access InCreateD3D12Access,
+	uint32 InAllocationAlignment,
+	const TCHAR* InName,
+	FD3D12ResourceLocation& ResourceLocation)
 {
 #if DO_CHECK
 	// Validate the create state
 	if (InHeapType == D3D12_HEAP_TYPE_READBACK)
 	{
-		check(InCreateState == D3D12_RESOURCE_STATE_COPY_DEST);
+		check(InCreateD3D12Access == ED3D12Access::CopyDest);
 	}
 	else if (InHeapType == D3D12_HEAP_TYPE_UPLOAD)
 	{
-		check(InCreateState == D3D12_RESOURCE_STATE_GENERIC_READ);
+		check(InCreateD3D12Access == ED3D12Access::GenericRead);
 	}
 	else if (InBufferUsage == BUF_UnorderedAccess && InResourceStateMode == ED3D12ResourceStateMode::SingleState)
 	{
-		check(InCreateState == D3D12_RESOURCE_STATE_UNORDERED_ACCESS);
+		check(InCreateD3D12Access == ED3D12Access::UAVMask);
 	}
 #if D3D12_RHI_RAYTRACING
 	else if (EnumHasAnyFlags(InBufferUsage, BUF_AccelerationStructure))
 	{
 		// RayTracing acceleration structures must be created in a particular state and may never transition out of it.
 		check(InResourceStateMode == ED3D12ResourceStateMode::SingleState);
-		check(InCreateState == D3D12_RESOURCE_STATE_RAYTRACING_ACCELERATION_STRUCTURE);
+		check(EnumHasAnyFlags(InCreateD3D12Access, ED3D12Access::BVHRead | ED3D12Access::BVHWrite));
 	}
 #endif // D3D12_RHI_RAYTRACING
 #endif  // DO_CHECK
 
-	AllocateResource(GetParentDevice()->GetGPUIndex(), InHeapType, InDesc, InDesc.Width, InAllocationAlignment, InResourceStateMode, InCreateState, nullptr, InName, ResourceLocation);
+	AllocateResource(
+		GetParentDevice()->GetGPUIndex(),
+		InHeapType,
+		InDesc,
+		InDesc.Width,
+		InAllocationAlignment,
+		InCreateD3D12Access,
+		InResourceStateMode,
+		InCreateD3D12Access,
+		nullptr,
+		InName,
+		ResourceLocation);
 }
 
 
-void FD3D12PoolAllocator::AllocateResource(uint32 GPUIndex, D3D12_HEAP_TYPE InHeapType, const FD3D12ResourceDesc& InDesc, uint64 InSize, uint32 InAllocationAlignment, ED3D12ResourceStateMode InResourceStateMode,
-	D3D12_RESOURCE_STATES InCreateState, const D3D12_CLEAR_VALUE* InClearValue, const TCHAR* InName, FD3D12ResourceLocation& ResourceLocation)
+void FD3D12PoolAllocator::AllocateResource(
+	uint32 GPUIndex,
+	D3D12_HEAP_TYPE InHeapType,
+	const FD3D12ResourceDesc& InDesc,
+	uint64 InSize,
+	uint32 InAllocationAlignment,
+	ED3D12Access InInitialD3D12Access,
+	ED3D12ResourceStateMode InResourceStateMode,
+	ED3D12Access InDefaultD3D12Access,
+	const D3D12_CLEAR_VALUE* InClearValue,
+	const TCHAR* InName,
+	FD3D12ResourceLocation& ResourceLocation)
 {
 	TRACE_CPUPROFILER_EVENT_SCOPE(D3D12RHI::AllocatePoolResource);
 
@@ -400,7 +439,16 @@ void FD3D12PoolAllocator::AllocateResource(uint32 GPUIndex, D3D12_HEAP_TYPE InHe
 			FD3D12ResourceDesc Desc = InDesc;
 			Desc.Alignment = AllocationAlignment;
 
-			FD3D12Resource* NewResource = CreatePlacedResource(AllocationData, Desc, InCreateState, InResourceStateMode, InClearValue, InName);
+			FD3D12Resource* NewResource =
+				CreatePlacedResource(
+					AllocationData,
+					InDesc,
+					InInitialD3D12Access,
+					InResourceStateMode,
+					InDefaultD3D12Access,
+					InClearValue,
+					InName);
+
 			ResourceLocation.SetResource(NewResource);
 			ResourceLocation.SetGPUVirtualAddress(NewResource->GetGPUVirtualAddress());
 		}
@@ -442,13 +490,38 @@ void FD3D12PoolAllocator::AllocateResource(uint32 GPUIndex, D3D12_HEAP_TYPE InHe
 				BackingHeap->BeginTrackingResidency(HeapDesc.SizeInBytes);
 			}
 
-			VERIFYD3D12RESULT(Adapter->CreatePlacedResource(InDesc, BackingHeap, 0, InCreateState, InResourceStateMode, InCreateState, InClearValue, &NewResource, InName));
+			VERIFYD3D12RESULT(
+				Adapter->CreatePlacedResource(
+					InDesc,
+					BackingHeap,
+					0,
+					InInitialD3D12Access,
+					InResourceStateMode,
+					InDefaultD3D12Access,
+					InClearValue,
+					&NewResource,
+					InName));
 		}
 		else
 		{
 			// Allocate Standalone - move to owner of resource because this allocator should only manage pooled allocations (needed for now to do the same as FD3D12DefaultBufferPool)
-			D3D12_RESOURCE_STATES DefaultState = (InResourceStateMode == ED3D12ResourceStateMode::Default) ? D3D12_RESOURCE_STATE_TBD : InCreateState;
-			VERIFYD3D12RESULT(Adapter->CreateCommittedResource(InDesc, GetGPUMask(), HeapProps, InCreateState, InResourceStateMode, DefaultState, InClearValue, &NewResource, InName, false));
+			const ED3D12Access DefaultAccess =
+				(InResourceStateMode == ED3D12ResourceStateMode::Default) ?
+					ED3D12Access::Unknown :
+					InDefaultD3D12Access;
+
+			VERIFYD3D12RESULT(
+				Adapter->CreateCommittedResource(
+					InDesc,
+					GetGPUMask(),
+					HeapProps,
+					InInitialD3D12Access,
+					InResourceStateMode,
+					DefaultAccess,
+					InClearValue,
+					&NewResource,
+					InName,
+					false));
 		}
 
 		ResourceLocation.AsStandAlone(NewResource, InSize);
@@ -457,7 +530,7 @@ void FD3D12PoolAllocator::AllocateResource(uint32 GPUIndex, D3D12_HEAP_TYPE InHe
 #if D3D12_RHI_RAYTRACING
 	// Elevates the raytracing acceleration structure heap priority, which may help performance / stability in low memory conditions.
 	// When allocation is pooled, the entire pool priority is boosted on creation.
-	if (InCreateState == D3D12_RESOURCE_STATE_RAYTRACING_ACCELERATION_STRUCTURE
+	if (EnumHasAnyFlags(InInitialD3D12Access, ED3D12Access::BVHRead | ED3D12Access::BVHWrite)
 		&& ResourceLocation.GetAllocatorType() != FD3D12ResourceLocation::AT_Pool)
 	{
 		ID3D12Pageable* HeapResource = ResourceLocation.GetResource()->GetPageable();
@@ -470,16 +543,27 @@ void FD3D12PoolAllocator::AllocateResource(uint32 GPUIndex, D3D12_HEAP_TYPE InHe
 FD3D12Resource* FD3D12PoolAllocator::CreatePlacedResource(
 	const FRHIPoolAllocationData& InAllocationData,
 	const FD3D12ResourceDesc& InDesc,
-	D3D12_RESOURCE_STATES InCreateState, 
-	ED3D12ResourceStateMode InResourceStateMode, 
-	const D3D12_CLEAR_VALUE* InClearValue, 
+	ED3D12Access InInitialD3D12Access,
+	ED3D12ResourceStateMode InResourceStateMode,
+	ED3D12Access InDefaultD3D12Access,
+	const D3D12_CLEAR_VALUE* InClearValue,
 	const TCHAR* InName)
 {
 	FD3D12Adapter* Adapter = GetParentDevice()->GetParentAdapter();
 	FD3D12HeapAndOffset HeapAndOffset = GetBackingHeapAndAllocationOffsetInBytes(InAllocationData);
 
 	FD3D12Resource* NewResource = nullptr;
-	VERIFYD3D12RESULT(Adapter->CreatePlacedResource(InDesc, HeapAndOffset.Heap, HeapAndOffset.Offset, InCreateState, InResourceStateMode, D3D12_RESOURCE_STATE_TBD, InClearValue, &NewResource, InName));
+	VERIFYD3D12RESULT(
+		Adapter->CreatePlacedResource(
+			InDesc,
+			HeapAndOffset.Heap,
+			HeapAndOffset.Offset,
+			InInitialD3D12Access,
+			InResourceStateMode,
+			InDefaultD3D12Access,
+			InClearValue,
+			&NewResource,
+			InName));
 	return NewResource;
 }
 
@@ -515,11 +599,11 @@ void FD3D12PoolAllocator::DeallocateResource(FD3D12ResourceLocation& ResourceLoc
 		// If still pending copy then clear the copy operation data
 		for (FD3D12VRAMCopyOperation& CopyOperation : PendingCopyOps)
 		{
-			if (CopyOperation.SourceResource == ResourceLocation.GetResource() ||
-				CopyOperation.DestResource == ResourceLocation.GetResource())
+			if (CopyOperation.SrcResource == ResourceLocation.GetResource() ||
+				CopyOperation.DstResource == ResourceLocation.GetResource())
 			{
-				CopyOperation.SourceResource = nullptr;
-				CopyOperation.DestResource = nullptr;
+				CopyOperation.SrcResource = nullptr;
+				CopyOperation.DstResource = nullptr;
 				break;
 			}
 		}
@@ -586,38 +670,31 @@ bool FD3D12PoolAllocator::HandleDefragRequest(FRHIContextArray const& Contexts,
 	uint64 CurrentOffset = Owner->GetOffsetFromBaseOfResource();
 	FD3D12Resource* CurrentResource = Owner->GetResource();
 
-	D3D12_RESOURCE_STATES SourceResourceState = D3D12_RESOURCE_STATE_TBD;
+	ED3D12Access SrcD3D12Access = ED3D12Access::Unknown;
 	if (CurrentResource->RequiresResourceStateTracking())
 	{
-
 		// State tracking is happening at the RHI high level 
-		ERHIAccess SourceResourceAccess = ERHIAccess::Unknown;
 		FD3D12Texture* CurrentResourceRHITexture = nullptr;
 		if (Owner->Owner)
 		{
 			if (CurrentResource->GetDesc().Dimension == D3D12_RESOURCE_DIMENSION_BUFFER)
 			{
 				FD3D12Buffer* CurrentResourceRHIBuffer = (FD3D12Buffer*)Owner->Owner;
-				SourceResourceAccess = CurrentResourceRHIBuffer->GetTrackedAccess_Unsafe();
+				SrcD3D12Access = ConvertToD3D12Access(CurrentResourceRHIBuffer->GetTrackedAccess_Unsafe());
 			}
 			else
 			{
 				CurrentResourceRHITexture = (FD3D12Texture*)Owner->Owner;
-				SourceResourceAccess = CurrentResourceRHITexture->GetTrackedAccess_Unsafe();
+				SrcD3D12Access = ConvertToD3D12Access(CurrentResourceRHITexture->GetTrackedAccess_Unsafe());
 			}
 		}
 
-		SourceResourceState = SourceResourceAccess == ERHIAccess::Unknown ? CurrentResource->GetInitialResourceState() : GetD3D12ResourceState(SourceResourceAccess, ED3D12QueueType::Direct, CurrentResourceRHITexture);
-
-		if (SourceResourceState == D3D12_RESOURCE_STATE_TBD)
-		{
-			SourceResourceState = CurrentResource->GetReadableState();
-		}
+		check(SrcD3D12Access != ED3D12Access::Unknown);
 	}
 	else
 	{
-		//If the resource is a single state type and doesn't need state tracking set the state as TBD so the scoped transition in FlushPendingCopyOps can revert back to the default state
-		SourceResourceState = D3D12_RESOURCE_STATE_TBD;
+		//If the resource is a single state type and doesn't need state tracking set the access as Unknown so the scoped transition in FlushPendingCopyOps can revert back to the default state
+		check(SrcD3D12Access == ED3D12Access::Unknown);
 	}
 
 	// Release the current allocation (will only be freed on the next frame fence)
@@ -631,20 +708,14 @@ bool FD3D12PoolAllocator::HandleDefragRequest(FRHIContextArray const& Contexts,
 	Owner->SetPoolAllocator(this);
 
 	// Notify owner of moved allocation data (recreated resources and SRVs if needed)
-	D3D12_RESOURCE_STATES DestCreateState = D3D12_RESOURCE_STATE_TBD;
-	Owner->OnAllocationMoved(Contexts, InSourceBlock, DestCreateState);
+	ED3D12Access DstD3D12Access = ED3D12Access::Unknown;
+	Owner->OnAllocationMoved(Contexts, InSourceBlock, DstD3D12Access);
 
-	if (Owner->GetResource()->RequiresResourceStateTracking())
-	{
-		if (DestCreateState == D3D12_RESOURCE_STATE_TBD)
-		{
-			DestCreateState = D3D12_RESOURCE_STATE_COMMON;
-		}
-	}
-	else
+	check(IsValidAccess(DstD3D12Access));
+	if (!Owner->GetResource()->RequiresResourceStateTracking())
 	{
 		//If the resource is a single state type and doesn't need state tracking set the state as TBD so the scoped transition in FlushPendingCopyOps can revert back to the default state
-		DestCreateState = D3D12_RESOURCE_STATE_TBD;
+		DstD3D12Access = ED3D12Access::Unknown;
 	}
 
 	// Add request to unlock the source block on the next fence value (copy operation should have been done by then)
@@ -657,18 +728,22 @@ bool FD3D12PoolAllocator::HandleDefragRequest(FRHIContextArray const& Contexts,
 	UnlockRequest.AllocationData = InSourceBlock;
 
 	// Schedule a copy operation of the actual data
-	FD3D12VRAMCopyOperation CopyOp;
-	CopyOp.SourceResource	= CurrentResource;
-	CopyOp.SourceResourceState = SourceResourceState;
-	CopyOp.SourceOffset		= CurrentOffset;
-	CopyOp.DestResource		= Owner->GetResource();
-	CopyOp.DestResourceState = DestCreateState;
-	CopyOp.DestOffset		= Owner->GetOffsetFromBaseOfResource();
-	CopyOp.Size				= InSourceBlock->GetSize();
-	CopyOp.CopyType			= AllocationStrategy == EResourceAllocationStrategy::kManualSubAllocation ? FD3D12VRAMCopyOperation::ECopyType::BufferRegion : FD3D12VRAMCopyOperation::ECopyType::Resource;
-	check(CopyOp.SourceResource != nullptr);
-	check(CopyOp.DestResource != nullptr);
-	PendingCopyOps.Add(CopyOp);
+	FD3D12VRAMCopyOperation CopyOp = {
+		.SrcResource          = CurrentResource,
+		.SrcResourcED3D12Access = SrcD3D12Access,
+		.SrcOffset            = static_cast<uint32>(CurrentOffset),
+		.DstResource          = Owner->GetResource(),
+		.DstResourcED3D12Access = DstD3D12Access,
+		.DstOffset            = static_cast<uint32>(Owner->GetOffsetFromBaseOfResource()),
+		.Size		          = InSourceBlock->GetSize(),
+		.CopyType	          =
+			(AllocationStrategy == EResourceAllocationStrategy::kManualSubAllocation) ?
+							FD3D12VRAMCopyOperation::ECopyType::BufferRegion :
+							FD3D12VRAMCopyOperation::ECopyType::Resource,
+	};
+	check(CopyOp.SrcResource != nullptr);
+	check(CopyOp.DstResource != nullptr);
+	PendingCopyOps.Add(MoveTemp(CopyOp));
 
 	UpdateAllocationTracking(*Owner, EAllocationType::Allocate);
 
@@ -886,19 +961,20 @@ void FD3D12PoolAllocator::FlushPendingCopyOps(FD3D12CommandContext& InCommandCon
 	for (FD3D12VRAMCopyOperation& CopyOperation : PendingCopyOps)
 	{
 		// Cleared copy op?
-		if (CopyOperation.SourceResource == nullptr || CopyOperation.DestResource == nullptr)
+		if (CopyOperation.SrcResource == nullptr || CopyOperation.DstResource == nullptr)
 		{
 			continue;
 		}
 
 #if D3D12_RHI_RAYTRACING
-		if (!CopyOperation.SourceResource->RequiresResourceStateTracking() && CopyOperation.SourceResource->GetDefaultResourceState() == D3D12_RESOURCE_STATE_RAYTRACING_ACCELERATION_STRUCTURE)
+		if (!CopyOperation.SrcResource->RequiresResourceStateTracking() && 
+			EnumHasAnyFlags(CopyOperation.SrcResource->GetDesc().Flags, D3D12_RESOURCE_FLAG_RAYTRACING_ACCELERATION_STRUCTURE))
 		{
 			// can't make state changes to RT resources
-			check(CopyOperation.DestResource->GetDefaultResourceState() == D3D12_RESOURCE_STATE_RAYTRACING_ACCELERATION_STRUCTURE);
+			check(EnumHasAnyFlags(CopyOperation.DstResource->GetDesc().Flags, D3D12_RESOURCE_FLAG_RAYTRACING_ACCELERATION_STRUCTURE));
 
-			D3D12_GPU_VIRTUAL_ADDRESS SrcAddress = CopyOperation.SourceResource->GetResource()->GetGPUVirtualAddress() + CopyOperation.SourceOffset;
-			D3D12_GPU_VIRTUAL_ADDRESS DestAddress = CopyOperation.DestResource->GetResource()->GetGPUVirtualAddress() + CopyOperation.DestOffset;
+			D3D12_GPU_VIRTUAL_ADDRESS SrcAddress = CopyOperation.SrcResource->GetResource()->GetGPUVirtualAddress() + CopyOperation.SrcOffset;
+			D3D12_GPU_VIRTUAL_ADDRESS DestAddress = CopyOperation.DstResource->GetResource()->GetGPUVirtualAddress() + CopyOperation.DstOffset;
 			InCommandContext.RayTracingCommandList()->CopyRaytracingAccelerationStructure(DestAddress, SrcAddress, D3D12_RAYTRACING_ACCELERATION_STRUCTURE_COPY_MODE_CLONE);
 		}
 		else
@@ -906,8 +982,18 @@ void FD3D12PoolAllocator::FlushPendingCopyOps(FD3D12CommandContext& InCommandCon
 		{
 
 			{
-				FScopedResourceBarrier ScopedResourceBarrierSrc(InCommandContext, CopyOperation.SourceResource, nullptr, D3D12_RESOURCE_STATE_COPY_SOURCE, D3D12_RESOURCE_BARRIER_ALL_SUBRESOURCES, CopyOperation.SourceResourceState);
-				FScopedResourceBarrier ScopedResourceBarrierDst(InCommandContext, CopyOperation.DestResource, nullptr, D3D12_RESOURCE_STATE_COPY_DEST, D3D12_RESOURCE_BARRIER_ALL_SUBRESOURCES, CopyOperation.DestResourceState);
+				FScopedResourceBarrier ScopedResourceBarrierSrc(
+					InCommandContext,
+					CopyOperation.SrcResource,
+					CopyOperation.SrcResourcED3D12Access,
+					ED3D12Access::CopySrc,
+					D3D12_RESOURCE_BARRIER_ALL_SUBRESOURCES);
+				FScopedResourceBarrier ScopedResourceBarrierDst(
+					InCommandContext,
+					CopyOperation.DstResource,
+					CopyOperation.DstResourcED3D12Access,
+					ED3D12Access::CopyDest,
+					D3D12_RESOURCE_BARRIER_ALL_SUBRESOURCES);
 
 				InCommandContext.FlushResourceBarriers();
 
@@ -916,10 +1002,10 @@ void FD3D12PoolAllocator::FlushPendingCopyOps(FD3D12CommandContext& InCommandCon
 				case FD3D12VRAMCopyOperation::BufferRegion:
 				{
 					InCommandContext.CopyBufferRegionChecked(
-						CopyOperation.DestResource->GetResource(), CopyOperation.DestResource->GetName(),
-						CopyOperation.DestOffset,
-						CopyOperation.SourceResource->GetResource(), CopyOperation.SourceResource->GetName(),
-						CopyOperation.SourceOffset,
+						CopyOperation.DstResource->GetResource(), CopyOperation.DstResource->GetName(),
+						CopyOperation.DstOffset,
+						CopyOperation.SrcResource->GetResource(), CopyOperation.SrcResource->GetName(),
+						CopyOperation.SrcOffset,
 						CopyOperation.Size
 					);
 					break;
@@ -927,8 +1013,8 @@ void FD3D12PoolAllocator::FlushPendingCopyOps(FD3D12CommandContext& InCommandCon
 				case FD3D12VRAMCopyOperation::Resource:
 				{
 					InCommandContext.GraphicsCommandList()->CopyResource(
-						CopyOperation.DestResource->GetResource(),
-						CopyOperation.SourceResource->GetResource()
+						CopyOperation.DstResource->GetResource(),
+						CopyOperation.SrcResource->GetResource()
 					);
 					break;
 				}
@@ -936,20 +1022,24 @@ void FD3D12PoolAllocator::FlushPendingCopyOps(FD3D12CommandContext& InCommandCon
 			}
 
 			// if the resource need tracking we are setting the state of the destination equal to the original source state
-			if (CopyOperation.SourceResource->RequiresResourceStateTracking())
+			if (CopyOperation.SrcResource->RequiresResourceStateTracking())
 			{
-				check(CopyOperation.DestResource->RequiresResourceStateTracking());
-				check(CopyOperation.SourceResourceState != D3D12_RESOURCE_STATE_TBD);
+				check(CopyOperation.DstResource->RequiresResourceStateTracking());
+				check(IsValidAccess(CopyOperation.SrcResourcED3D12Access));
 
-				if (CopyOperation.SourceResourceState != CopyOperation.DestResourceState)
+				if (CopyOperation.SrcResourcED3D12Access != CopyOperation.DstResourcED3D12Access)
 				{
-					InCommandContext.AddTransitionBarrier(CopyOperation.DestResource, CopyOperation.DestResourceState, CopyOperation.SourceResourceState, D3D12_RESOURCE_BARRIER_ALL_SUBRESOURCES);
+					InCommandContext.AddBarrier(
+						CopyOperation.DstResource,
+						CopyOperation.DstResourcED3D12Access,
+						CopyOperation.SrcResourcED3D12Access,
+						D3D12_RESOURCE_BARRIER_ALL_SUBRESOURCES);
 				}
 			}
 		}
 
-		InCommandContext.UpdateResidency(CopyOperation.SourceResource);
-		InCommandContext.UpdateResidency(CopyOperation.DestResource);
+		InCommandContext.UpdateResidency(CopyOperation.SrcResource);
+		InCommandContext.UpdateResidency(CopyOperation.DstResource);
 	}
 
 	InCommandContext.ConditionalSplitCommandList();
diff --git a/Engine/Source/Runtime/D3D12RHI/Private/D3D12PoolAllocator.h b/Engine/Source/Runtime/D3D12RHI/Private/D3D12PoolAllocator.h
--- a/Engine/Source/Runtime/D3D12RHI/Private/D3D12PoolAllocator.h
+++ b/Engine/Source/Runtime/D3D12RHI/Private/D3D12PoolAllocator.h
@@ -28,7 +28,7 @@ struct FD3D12ResourceInitConfig
 		Config.HeapType = D3D12_HEAP_TYPE_UPLOAD;
 		Config.HeapFlags = D3D12_HEAP_FLAG_ALLOW_ONLY_BUFFERS;
 		Config.ResourceFlags = D3D12_RESOURCE_FLAG_NONE;
-		Config.InitialResourceState = D3D12_RESOURCE_STATE_GENERIC_READ;
+		Config.InitialD3D12Access = ED3D12Access::GenericRead;
 		return Config;
 	}
 
@@ -37,13 +37,13 @@ struct FD3D12ResourceInitConfig
 		return HeapType == InOther.HeapType &&
 			HeapFlags == InOther.HeapFlags &&
 			ResourceFlags == InOther.ResourceFlags &&
-			InitialResourceState == InOther.InitialResourceState;
+			InitialD3D12Access   == InOther.InitialD3D12Access;
 	}
 
 	D3D12_HEAP_TYPE HeapType;
 	D3D12_HEAP_FLAGS HeapFlags;
 	D3D12_RESOURCE_FLAGS ResourceFlags;
-	D3D12_RESOURCE_STATES InitialResourceState;
+	ED3D12Access InitialD3D12Access;
 };
 
 
@@ -58,12 +58,12 @@ struct FD3D12VRAMCopyOperation
 		Resource,
 	};
 
-	FD3D12Resource* SourceResource;
-	D3D12_RESOURCE_STATES SourceResourceState;
-	uint32 SourceOffset;
-	FD3D12Resource* DestResource;
-	D3D12_RESOURCE_STATES DestResourceState;
-	uint32 DestOffset;
+	FD3D12Resource* SrcResource;
+	ED3D12Access SrcResourcED3D12Access;
+	uint32 SrcOffset;
+	FD3D12Resource* DstResource;
+	ED3D12Access DstResourcED3D12Access;
+	uint32 DstOffset;
 	uint32 Size;
 	ECopyType CopyType;
 };
@@ -137,13 +137,30 @@ public:
 	
 	// Function names currently have to match with FD3D12DefaultBufferPool until we can make full replacement of the allocator
 	bool SupportsAllocation(D3D12_HEAP_TYPE InHeapType, D3D12_RESOURCE_FLAGS InResourceFlags, EBufferUsageFlags InBufferUsage, ED3D12ResourceStateMode InResourceStateMode, uint32 Alignment) const;
-	void AllocDefaultResource(D3D12_HEAP_TYPE InHeapType, const FD3D12ResourceDesc& InDesc, EBufferUsageFlags InBufferUsage, ED3D12ResourceStateMode InResourceStateMode,
-		D3D12_RESOURCE_STATES InCreateState, uint32 InAlignment, const TCHAR* InName, FD3D12ResourceLocation& ResourceLocation);
+	void AllocDefaultResource(
+		D3D12_HEAP_TYPE InHeapType,
+		const FD3D12ResourceDesc& InDesc,
+		EBufferUsageFlags InBufferUsage,
+		ED3D12ResourceStateMode InResourceStateMode,
+		ED3D12Access InCreateD3D12Access,
+		uint32 InAlignment,
+		const TCHAR* InName,
+		FD3D12ResourceLocation& ResourceLocation);
 	virtual void DeallocateResource(FD3D12ResourceLocation& ResourceLocation, bool bDefragFree = false);
 
 	// Implementation of ID3D12ResourceAllocator
-	virtual void AllocateResource(uint32 GPUIndex, D3D12_HEAP_TYPE InHeapType, const FD3D12ResourceDesc& InDesc, uint64 InSize, uint32 InAllocationAlignment, ED3D12ResourceStateMode InResourceStateMode,
-		D3D12_RESOURCE_STATES InCreateState, const D3D12_CLEAR_VALUE* InClearValue, const TCHAR* InName, FD3D12ResourceLocation& ResourceLocation) override;
+	virtual void AllocateResource(
+		uint32 GPUIndex,
+		D3D12_HEAP_TYPE InHeapType,
+		const FD3D12ResourceDesc& InDesc,
+		uint64 InSize,
+		uint32 InAllocationAlignment,
+		ED3D12Access InInitialD3D12Access,
+		ED3D12ResourceStateMode InResourceStateMode,
+		ED3D12Access InDefaultD3D12Access,
+		const D3D12_CLEAR_VALUE* InClearValue,
+		const TCHAR* InName,
+		FD3D12ResourceLocation& ResourceLocation) override;
 
 	void CleanUpAllocations(uint64 InFrameLag, bool bForceFree = false);
 	void FlushPendingCopyOps(FD3D12CommandContext& InCommandContext);
@@ -167,7 +184,14 @@ protected:
 	virtual bool HandleDefragRequest(FRHIContextArray const& Contexts, FRHIPoolAllocationData* InSourceBlock, FRHIPoolAllocationData& InTmpTargetBlock) override;
 	
 	// Placed resource allocation helper function which can be overriden
-	virtual FD3D12Resource* CreatePlacedResource(const FRHIPoolAllocationData& InAllocationData, const FD3D12ResourceDesc& InDesc, D3D12_RESOURCE_STATES InCreateState, ED3D12ResourceStateMode InResourceStateMode, const D3D12_CLEAR_VALUE* InClearValue, const TCHAR* InName);
+	virtual FD3D12Resource* CreatePlacedResource(
+		const FRHIPoolAllocationData& InAllocationData,
+		const FD3D12ResourceDesc& InDesc,
+		ED3D12Access InInitialD3D12Access,
+		ED3D12ResourceStateMode InResourceStateMode,
+		ED3D12Access InDefaultD3D12Access,
+		const D3D12_CLEAR_VALUE* InClearValue,
+		const TCHAR* InName);
 	
 	// Track the allocation
 	enum class EAllocationType
diff --git a/Engine/Source/Runtime/D3D12RHI/Private/D3D12Query.cpp b/Engine/Source/Runtime/D3D12RHI/Private/D3D12Query.cpp
--- a/Engine/Source/Runtime/D3D12RHI/Private/D3D12Query.cpp
+++ b/Engine/Source/Runtime/D3D12RHI/Private/D3D12Query.cpp
@@ -106,14 +106,15 @@ FD3D12QueryHeap::FD3D12QueryHeap(FD3D12Device* Device, D3D12_QUERY_TYPE QueryTyp
 		D3D12_RESOURCE_DESC BufferDesc = CD3DX12_RESOURCE_DESC::Buffer(GetResultSize() * NumQueries);
 
 		// Create the readback heap
-		VERIFYD3D12RESULT(Device->GetParentAdapter()->CreateCommittedResource(
-			BufferDesc,
-			GetGPUMask(),
-			HeapProps,
-			D3D12_RESOURCE_STATE_COPY_DEST,
-			nullptr,
-			ResultBuffer.GetInitReference(),
-			ResultBufferName
+		VERIFYD3D12RESULT(
+			Device->GetParentAdapter()->CreateCommittedResource(
+				BufferDesc,
+				GetGPUMask(),
+				HeapProps,
+				ED3D12Access::CopyDest,
+				nullptr,
+				ResultBuffer.GetInitReference(),
+				ResultBufferName
 		));
 		SetD3D12ResourceName(ResultBuffer, ResultBufferName);
 	}
diff --git a/Engine/Source/Runtime/D3D12RHI/Private/D3D12RHI.cpp b/Engine/Source/Runtime/D3D12RHI/Private/D3D12RHI.cpp
--- a/Engine/Source/Runtime/D3D12RHI/Private/D3D12RHI.cpp
+++ b/Engine/Source/Runtime/D3D12RHI/Private/D3D12RHI.cpp
@@ -51,6 +51,7 @@ static FAutoConsoleVariableRef CVarPSOPrecacheD3D12DriverCacheAware(
 	ECVF_ReadOnly);
 
 FD3D12DynamicRHI* FD3D12DynamicRHI::SingleD3DRHI = nullptr;
+bool FD3D12DynamicRHI::bFormatAliasedTexturesMustBeCreatedUsingCommonLayout = false;
 
 FD3D12WorkaroundFlags GD3D12WorkaroundFlags;
 
@@ -100,9 +101,6 @@ FD3D12DynamicRHI::FD3D12DynamicRHI(const TArray<TSharedPtr<FD3D12Adapter>>& Chos
 	GTexturePoolSize = 0;
 	GConfig->GetInt(TEXT("TextureStreaming"), TEXT("PoolSizeVRAMPercentage"), GPoolSizeVRAMPercentage, GEngineIni);
 
-	GRHITransitionPrivateData_SizeInBytes = sizeof(FD3D12TransitionData);
-	GRHITransitionPrivateData_AlignInBytes = alignof(FD3D12TransitionData);
-
 	// Initialize the platform pixel format map.
 	GPixelFormats[PF_Unknown		].PlatformFormat = DXGI_FORMAT_UNKNOWN;
 	GPixelFormats[PF_A32B32G32R32F	].PlatformFormat = DXGI_FORMAT_R32G32B32A32_FLOAT;
@@ -254,8 +252,6 @@ FD3D12DynamicRHI::FD3D12DynamicRHI(const TArray<TSharedPtr<FD3D12Adapter>>& Chos
 	GRHIGlobals.SupportsVertexShaderUAVs = true;
 
 	GRHIGlobals.NeedsExtraTransitions = true;
-	GRHIGlobals.NeedsTransientDiscardStateTracking = true;
-	GRHIGlobals.NeedsTransientDiscardOnGraphicsWorkaround = true;
 }
 
 void FD3D12DynamicRHI::PostInit()
@@ -756,9 +752,7 @@ void FD3D12DynamicRHI::RHIFinishExternalComputeWork(FRHICommandListBase& Executi
 
 void FD3D12DynamicRHI::RHITransitionResource(FRHICommandList& RHICmdList, FRHITexture* InTexture, D3D12_RESOURCE_STATES InState, uint32 InSubResource)
 {
-	FString StateString = ConvertToResourceStateString(InState);
-
-	UE_LOG(LogD3D12RHI, Error, TEXT("RHITransitionResource cannot transition from Unknown state to %s state on Texture %s subresource %d. Use RHICmdList.Transition() instead."), *StateString, *InTexture->GetName().GetPlainNameString(), InSubResource);
+	UE_LOG(LogD3D12RHI, Error, TEXT("RHITransitionResource is no longer functional. Please use RHICmdList.Transition() instead."));
 }
 
 void FD3D12DynamicRHI::RHISignalManualFence(FRHICommandList& RHICmdList, ID3D12Fence* Fence, uint64 Value)
diff --git a/Engine/Source/Runtime/D3D12RHI/Private/D3D12RHIDefinitions.h b/Engine/Source/Runtime/D3D12RHI/Private/D3D12RHIDefinitions.h
--- a/Engine/Source/Runtime/D3D12RHI/Private/D3D12RHIDefinitions.h
+++ b/Engine/Source/Runtime/D3D12RHI/Private/D3D12RHIDefinitions.h
@@ -17,6 +17,9 @@
 #define DEBUG_RESOURCE_STATES	0
 #define DX12_RESOURCE_NAME_TO_LOG	""
 
+// Useful for debugging. Only flushes the default context though.
+#define D3D12RHI_IDLE_AFTER_EVERY_GPU_EVENT 0
+
 #define D3D12_RHI_RAYTRACING (RHI_RAYTRACING)
 
 #if D3D12_MAX_COMMANDLIST_INTERFACE >= 10
@@ -27,6 +30,22 @@
 	#define D3D12_RHI_WORKGRAPHS_GRAPHICS 0
 #endif
 
+#ifndef D3D12RHI_ALLOW_ENHANCED_BARRIERS
+	#error "Please define D3D12RHI_ALLOW_ENHANCED_BARRIERS in <Platform>D3D12RHIDefinitions.h!"
+#endif
+
+#ifndef D3D12RHI_ALLOW_LEGACY_BARRIERS
+	#error "Please define D3D12RHI_ALLOW_LEGACY_BARRIERS in <Platform>D3D12RHIDefinitions.h!"
+#endif
+
+#if D3D12_MAX_DEVICE_INTERFACE >= 11
+	#define D3D12RHI_SUPPORTS_ENHANCED_BARRIERS (D3D12RHI_ALLOW_ENHANCED_BARRIERS)
+#else
+	#define D3D12RHI_SUPPORTS_ENHANCED_BARRIERS 0
+#endif
+
+#define D3D12RHI_SUPPORTS_LEGACY_BARRIERS (D3D12RHI_ALLOW_LEGACY_BARRIERS)
+
 #if D3D12_MAX_DEVICE_INTERFACE >= 12
 	#define D3D12RHI_SUPPORTS_UNCOMPRESSED_UAV 1
 #else
@@ -66,9 +85,9 @@
 #endif
 
 #if EXECUTE_DEBUG_COMMAND_LISTS
-	#define DEBUG_EXECUTE_COMMAND_LIST(scope) if (scope##->ActiveQueries == 0) { scope##->FlushCommands(ED3D12FlushFlags::WaitForCompletion); }
-	#define DEBUG_EXECUTE_COMMAND_CONTEXT(context) if (context.ActiveQueries == 0) { context##.FlushCommands(ED3D12FlushFlags::WaitForCompletion); }
-	#define DEBUG_RHI_EXECUTE_COMMAND_LIST(scope) if (scope##->GetRHIDevice(0)->GetDefaultCommandContext().ActiveQueries == 0) { scope##->GetRHIDevice(0)->GetDefaultCommandContext().FlushCommands(ED3D12FlushFlags::WaitForCompletion); }
+	#define DEBUG_EXECUTE_COMMAND_LIST(scope) if ((scope)->ActiveQueries == 0) { (scope)->FlushCommands(ED3D12FlushFlags::WaitForCompletion); }
+	#define DEBUG_EXECUTE_COMMAND_CONTEXT(context) if ((context).ActiveQueries == 0) { (context).FlushCommands(ED3D12FlushFlags::WaitForCompletion); }
+	#define DEBUG_RHI_EXECUTE_COMMAND_LIST(scope) if ((scope)->GetRHIDevice(0)->GetDefaultCommandContext().ActiveQueries == 0) { (scope)->GetRHIDevice(0)->GetDefaultCommandContext().FlushCommands(ED3D12FlushFlags::WaitForCompletion); }
 #else
 	#define DEBUG_EXECUTE_COMMAND_LIST(scope) 
 	#define DEBUG_EXECUTE_COMMAND_CONTEXT(context) 
diff --git a/Engine/Source/Runtime/D3D12RHI/Private/D3D12RHIPrivate.h b/Engine/Source/Runtime/D3D12RHI/Private/D3D12RHIPrivate.h
--- a/Engine/Source/Runtime/D3D12RHI/Private/D3D12RHIPrivate.h
+++ b/Engine/Source/Runtime/D3D12RHI/Private/D3D12RHIPrivate.h
@@ -62,6 +62,7 @@
 #include "D3D12Stats.h"
 #include "D3D12Device.h"
 #include "D3D12Adapter.h"
+#include "D3D12Access.h"
 
 template< typename t_A, typename t_B >
 inline t_A RoundUpToNextMultiple(const t_A& a, const t_B& b)
@@ -160,6 +161,7 @@ class FD3D12DynamicRHI : public ID3D12PlatformDynamicRHI
 	friend class FD3D12CommandContext;
 
 	static FD3D12DynamicRHI* SingleD3DRHI;
+	static bool bFormatAliasedTexturesMustBeCreatedUsingCommonLayout;
 
 public:
 
@@ -518,7 +520,13 @@ public:
 public:
 
 #if	PLATFORM_SUPPORTS_VIRTUAL_TEXTURES
-	virtual void* CreateVirtualTexture(ETextureCreateFlags InFlags, D3D12_RESOURCE_DESC& ResourceDesc, const struct FD3D12TextureLayout& TextureLayout, FD3D12Resource** ppResource, FPlatformMemory::FPlatformVirtualMemoryBlock& RawTextureBlock, D3D12_RESOURCE_STATES InitialUsage = D3D12_RESOURCE_STATE_PIXEL_SHADER_RESOURCE) = 0;
+	virtual void* CreateVirtualTexture(
+		ETextureCreateFlags InFlags,
+		D3D12_RESOURCE_DESC& ResourceDesc,
+		const struct FD3D12TextureLayout& TextureLayout,
+		FD3D12Resource** ppResource,
+		FPlatformMemory::FPlatformVirtualMemoryBlock& RawTextureBlock,
+		ED3D12Access InitialD3D12Access) = 0;
 	virtual void DestroyVirtualTexture(ETextureCreateFlags InFlags, void* RawTextureMemory, FPlatformMemory::FPlatformVirtualMemoryBlock& RawTextureBlock, uint64 CommittedTextureSize) = 0;
 #endif
 	virtual bool HandleSpecialLock(FRHILockTextureResult& Result, FD3D12Texture* InTexture, const FRHILockTextureArgs& Arguments) { return false; }
@@ -579,13 +587,15 @@ protected:
 
 public:
 
+	static void SetFormatAliasedTexturesMustBeCreatedUsingCommonLayout(bool bValue);
+
 	virtual FD3D12ResourceDesc GetResourceDesc(const FRHITextureDesc& CreateInfo) const;
 
 	struct FCreateTextureInternalResult
 	{
 		FD3D12Texture* Texture{};
-		D3D12_RESOURCE_STATES CreateState{};
-		D3D12_RESOURCE_STATES DesiredState{};
+		ED3D12Access CreateD3D12Access{};
+		ED3D12Access DesiredD3D12Access{};
 	};
 
 	// Per-platform D3D12Texture creation.
@@ -597,7 +607,7 @@ public:
 		check(CreateDesc.BulkData == nullptr);
 
 		FCreateTextureInternalResult CreateResult = CreateTextureInternal(CreateDesc, ResourceAllocator);
-		checkSlow(CreateResult.CreateState == CreateResult.DesiredState);
+		checkSlow(CreateResult.CreateD3D12Access == CreateResult.DesiredD3D12Access);
 		return CreateResult.Texture;
 	}
 
@@ -606,7 +616,7 @@ public:
 	struct FCreateBufferInternalResult
 	{
 		FD3D12Buffer* Buffer{};
-		D3D12_RESOURCE_STATES DesiredState{};
+		ED3D12Access DesiredD3D12Access{};
 	};
 
 	FCreateBufferInternalResult CreateBufferInternal(const FRHIBufferCreateDesc& CreateDesc, bool bHasInitialData, ID3D12ResourceAllocator* ResourceAllocator);
@@ -707,45 +717,60 @@ private:
 class FScopedResourceBarrier
 {
 private:
-	FD3D12ContextCommon&        Context;
-	FD3D12Resource* const Resource;
-	FD3D12ResourceLocation* const ResourceLocation;
-	D3D12_RESOURCE_STATES const DesiredState;
-	D3D12_RESOURCE_STATES  BeforeState;
-	uint32                const Subresource;
-
-	bool bRestoreState = false;
+	FD3D12ContextCommon&  Context;
+	const FD3D12Resource* pResource;
+	ED3D12Access BeforeD3D12Access;
+	const ED3D12Access DesiredD3D12Access;
+	const uint32 Subresource;
+	bool bRestoreState;
 
 public:
-	FScopedResourceBarrier(FD3D12ContextCommon& Context, FD3D12Resource* Resource, FD3D12ResourceLocation* InResourceLocation, D3D12_RESOURCE_STATES DesiredState, uint32 Subresource, D3D12_RESOURCE_STATES BeforeState = D3D12_RESOURCE_STATE_TBD)
-		: Context     (Context)
-		, Resource(Resource)
-		, ResourceLocation(InResourceLocation)
-		, DesiredState(DesiredState)
-		, BeforeState(BeforeState)
-		, Subresource (Subresource)
+	FScopedResourceBarrier(
+		FD3D12ContextCommon& InContext,
+		const FD3D12Resource* InResource,
+		ED3D12Access InBeforeD3D12Access,
+		ED3D12Access InDesiredD3D12Access,
+		uint32 InSubresource)
+	: Context(InContext)
+	, pResource(InResource)
+	, BeforeD3D12Access(InBeforeD3D12Access)
+	, DesiredD3D12Access(InDesiredD3D12Access)
+	, Subresource (InSubresource)
+	, bRestoreState(false)
 	{
-		if ((!Resource->RequiresResourceStateTracking()) || (BeforeState != D3D12_RESOURCE_STATE_TBD))
-		{
-			// Resources that do not require state tracking are kept in a default state, meaning the previous state for the transition is known.
-			D3D12_RESOURCE_STATES CurrentState = (BeforeState != D3D12_RESOURCE_STATE_TBD) ? BeforeState : Resource->GetDefaultResourceState();
+		check(IsValidAccess(DesiredD3D12Access));
+		check(DesiredD3D12Access != ED3D12Access::Unknown);
 
-			// Some states such as D3D12_RESOURCE_STATE_GENERIC_READ already includes D3D12_RESOURCE_STATE_COPY_SOURCE as well as other states, therefore transition isn't required.
-			if (CurrentState != DesiredState && !EnumHasAllFlags(CurrentState, DesiredState))
-			{
-				Context.AddTransitionBarrier(Resource, CurrentState, DesiredState, Subresource);
+		const bool bBeforED3D12AccessIsKnown = BeforeD3D12Access != ED3D12Access::Unknown;
+		if (!pResource->RequiresResourceStateTracking() || bBeforED3D12AccessIsKnown)
+		{
+			BeforeD3D12Access = bBeforED3D12AccessIsKnown ? BeforeD3D12Access : pResource->GetDefaultAccess();
+			check(IsValidAccess(BeforeD3D12Access));
+			check(BeforeD3D12Access != ED3D12Access::Unknown);
 
-				// We will add a transition, we need to transition back to the default state when the scoped object dies.
-				bRestoreState = true;
-			}
+			Context.AddBarrier(pResource, BeforeD3D12Access, DesiredD3D12Access, Subresource);
+			bRestoreState = true;
 		}
 	}
 
+	FScopedResourceBarrier(
+		FD3D12ContextCommon& InContext,
+		const FD3D12Resource* InResource,
+		ED3D12Access InDesiredD3D12Access,
+		uint32 InSubresource)
+	: FScopedResourceBarrier(
+		InContext,
+		InResource,
+		ED3D12Access::Unknown,
+		InDesiredD3D12Access,
+		InSubresource)
+	{}
+
 	~FScopedResourceBarrier()
 	{
 		if (bRestoreState)
 		{
-			Context.AddTransitionBarrier(Resource, DesiredState, (BeforeState != D3D12_RESOURCE_STATE_TBD) ? BeforeState : Resource->GetDefaultResourceState(), Subresource);
+			Context.AddBarrier(pResource, DesiredD3D12Access, BeforeD3D12Access, Subresource);
 		}
 	}
 };
diff --git a/Engine/Source/Runtime/D3D12RHI/Private/D3D12RayTracing.cpp b/Engine/Source/Runtime/D3D12RHI/Private/D3D12RayTracing.cpp
--- a/Engine/Source/Runtime/D3D12RHI/Private/D3D12RayTracing.cpp
+++ b/Engine/Source/Runtime/D3D12RHI/Private/D3D12RayTracing.cpp
@@ -739,7 +739,7 @@ FD3D12RayTracingCompactionRequestHandler::FD3D12RayTracingCompactionRequestHandl
 
 	const FRHIBufferCreateDesc CreateDesc =
 		FRHIBufferCreateDesc::Create(TEXT("PostBuildInfoBuffer"), BufferSize, 0, BUF_UnorderedAccess | BUF_SourceCopy)
-		.SetInitialState(ERHIAccess::UAVMask)
+		.SetInitialState(ERHIAccess::UAVCompute)
 		.SetGPUMask(FRHIGPUMask::FromIndex(GetParentDevice()->GetGPUIndex()));
 
 	PostBuildInfoBuffer = GetParentDevice()->GetParentAdapter()->CreateRHIBuffer(
@@ -747,7 +747,7 @@ FD3D12RayTracingCompactionRequestHandler::FD3D12RayTracingCompactionRequestHandl
 		8,
 		CreateDesc,
 		ED3D12ResourceStateMode::MultiState,
-		D3D12_RESOURCE_STATE_UNORDERED_ACCESS
+		ED3D12Access::UAVCompute
 	);
 	SetD3D12ResourceName(PostBuildInfoBuffer->GetResource(), TEXT("PostBuildInfoBuffer"));
 
@@ -855,19 +855,32 @@ void FD3D12RayTracingCompactionRequestHandler::Update(FD3D12CommandContext& Cont
 		//PostBuildInfoBuffer enters in D3D12_RESOURCE_STATE_UNORDERED_ACCESS
 		//Context.TransitionResource(PostBuildInfoBuffer->GetResource(), D3D12_RESOURCE_STATE_COPY_SOURCE, D3D12_RESOURCE_STATE_UNORDERED_ACCESS, 0);
 
+		// @TODO - Might do better if we had a way to say "BVHPostBuildInfoRead"
 		// Force UAV barrier to make sure all previous builds ops are finished
-		Context.AddUAVBarrier();
+		Context.AddGlobalBarrier(
+			ED3D12Access::BVHWrite,
+			ED3D12Access::BVHRead);
+
+		// Must flush barriers before making a direct call to a D3D API
 		Context.FlushResourceBarriers();
 
 		// Emit the RT post build info from the selected requests
 		Context.RayTracingCommandList()->EmitRaytracingAccelerationStructurePostbuildInfo(&PostBuildInfoDesc, ActiveBLASGPUAddresses.Num(), ActiveBLASGPUAddresses.GetData());
 
 		// Transition to copy source and perform the copy to readback
-		Context.TransitionResource(PostBuildInfoBuffer->GetResource(), D3D12_RESOURCE_STATE_UNORDERED_ACCESS, D3D12_RESOURCE_STATE_COPY_SOURCE, 0);
-		Context.FlushResourceBarriers();
+		Context.AddBarrier(
+			PostBuildInfoBuffer->GetResource(),
+			ED3D12Access::UAVCompute,
+			ED3D12Access::CopySrc,
+			0);
 
 		Context.RHICopyToStagingBuffer(PostBuildInfoBuffer, PostBuildInfoStagingBuffer, 0, sizeof(uint64) * ActiveBLASGPUAddresses.Num());
-		Context.TransitionResource(PostBuildInfoBuffer->GetResource(), D3D12_RESOURCE_STATE_COPY_SOURCE, D3D12_RESOURCE_STATE_UNORDERED_ACCESS, 0);
+
+		Context.AddBarrier(
+			PostBuildInfoBuffer->GetResource(),
+			ED3D12Access::CopySrc,
+			ED3D12Access::UAVCompute,
+			0);
 
 		// Update the sync point
 		PostBuildInfoBufferReadbackSyncPoint = Context.GetContextSyncPoint();
@@ -1703,7 +1716,7 @@ public:
 				BufferDesc.Alignment,
 				CreateDesc,
 				ED3D12ResourceStateMode::MultiState,
-				D3D12_RESOURCE_STATE_COPY_DEST,
+				ED3D12Access::CopyDest,
 				/*bHasInitialData*/ true
 			);
 
@@ -1711,10 +1724,10 @@ public:
 			Context.BatchedSyncPoints.ToWait.Emplace(Buffer->UploadResourceDataViaCopyQueue(Context, &Data));
 
 			// Enqueue transition to SRV
-			Context.TransitionResource(
+			Context.AddBarrier(
 				Buffer->GetResource(),
-				D3D12_RESOURCE_STATE_COPY_DEST,
-				D3D12_RESOURCE_STATE_NON_PIXEL_SHADER_RESOURCE,
+				ED3D12Access::CopyDest,
+				ED3D12Access::SRVCompute,
 				0);
 		}
 
@@ -3523,28 +3536,30 @@ static TRefCountPtr<FD3D12Buffer> CreateRayTracingBuffer(FD3D12Adapter* Adapter,
 
 	if (Type == ERayTracingBufferType::AccelerationStructure)
 	{
-		const D3D12_RESOURCE_DESC ResourceDesc = CD3DX12_RESOURCE_DESC::Buffer(Size, D3D12_RESOURCE_FLAG_NONE);
+		const D3D12_RESOURCE_DESC ResourceDesc = CD3DX12_RESOURCE_DESC::Buffer(Size, D3D12_RESOURCE_FLAG_RAYTRACING_ACCELERATION_STRUCTURE);
 		const FRHIBufferCreateDesc CreateDesc =
 			FRHIBufferCreateDesc::Create(*DebugNameString, Size, 0, BUF_AccelerationStructure)
 			.SetInitialState(ERHIAccess::BVHWrite)
 			.SetGPUMask(GPUMask);
 
+		// TODO: Manage switching between RAS_READ and RAS_WRITE?
 		Result = Adapter->CreateRHIBuffer(
 			ResourceDesc,
 			D3D12_RAYTRACING_ACCELERATION_STRUCTURE_BYTE_ALIGNMENT,
 			CreateDesc,
 			ED3D12ResourceStateMode::SingleState,
-			D3D12_RESOURCE_STATE_RAYTRACING_ACCELERATION_STRUCTURE
+			ED3D12Access::BVHRead | ED3D12Access::BVHWrite
 		);
 	}
 	else if (Type == ERayTracingBufferType::Scratch)
 	{
+		// 
 		// Scratch doesn't need single state anymore because there are only a few scratch allocations left and allocating a 
 		// dedicated single state heap for it wastes memory - ideally all scratch allocations should be transient
 		const D3D12_RESOURCE_DESC ResourceDesc = CD3DX12_RESOURCE_DESC::Buffer(Size, D3D12_RESOURCE_FLAG_ALLOW_UNORDERED_ACCESS);
 		const FRHIBufferCreateDesc CreateDesc =
 			FRHIBufferCreateDesc::Create(*DebugNameString, Size, 0, BUF_UnorderedAccess)
-			.SetInitialState(ERHIAccess::BVHWrite)
+			.SetInitialState(ERHIAccess::UAVCompute)
 			.SetGPUMask(GPUMask);
 
 		Result = Adapter->CreateRHIBuffer(
@@ -3552,7 +3567,7 @@ static TRefCountPtr<FD3D12Buffer> CreateRayTracingBuffer(FD3D12Adapter* Adapter,
 			D3D12_RAYTRACING_ACCELERATION_STRUCTURE_BYTE_ALIGNMENT,
 			CreateDesc,
 			ED3D12ResourceStateMode::Default,
-			D3D12_RESOURCE_STATE_UNORDERED_ACCESS
+			ED3D12Access::UAVCompute
 		);
 
 		// Elevates the scratch buffer heap priority, which may help performance / stability in low memory conditions 
@@ -4535,7 +4550,7 @@ void PrepareAccelerationStructureBuild(
 		uint32 InstanceBufferStride = GRHIRayTracingInstanceDescriptorSize;
 		PRAGMA_DISABLE_DEPRECATION_WARNINGS
 		// TODO: Validation related to SBT needs to be done somewhere else since SBT is not known when in BuildAccelerationStructure
-			uint32 TotalHitGroupSlots = Scene.Initializer.NumTotalSegments;
+		uint32 TotalHitGroupSlots = Scene.Initializer.NumTotalSegments;
 		PRAGMA_ENABLE_DEPRECATION_WARNINGS
 		FRayTracingValidateSceneBuildParamsCS::Dispatch(RHICmdList,
 			TotalHitGroupSlots, NumInstances,
@@ -4763,8 +4778,11 @@ void FD3D12CommandContext::RHIBuildAccelerationStructures(TConstArrayView<FRayTr
 	// This is required because there are currently no explicit read/write barriers
 	// for acceleration structures, but we need to ensure that all commands
 	// are complete before BLAS is used again on the GPU.
+	//
+	// @TODO - EBs let us express those R/W barriers so we should do that
+	//		   Also could use the postbuild flags here if added (see CreateAccelerationStructureBuildDesc())
 
-	AddUAVBarrier();
+	AddGlobalBarrier(ED3D12Access::BVHWrite, ED3D12Access::BVHRead);
 }
 
 void FD3D12CommandContext::RHIBuildAccelerationStructures(TConstArrayView<FRayTracingSceneBuildParams> Params)
@@ -4801,15 +4819,16 @@ void FD3D12CommandContext::RHIBuildAccelerationStructures(TConstArrayView<FRayTr
 		);
 	}
 
-	// UAV barrier is used here to ensure that all bottom level acceleration structures are built
-	AddUAVBarrier();
+	// Barrier is used here to ensure that all bottom level acceleration structures are built
+	AddGlobalBarrier(ED3D12Access::BVHWrite, ED3D12Access::BVHRead);
+
 	FlushResourceBarriers();
 
 	BuildAccelerationStructuresInternal(BuildDescs);
 
-	// UAV barrier is used here to ensure that the acceleration structure build is complete before any rays are traced
+	// Barrier is used here to ensure that the acceleration structure build is complete before any rays are traced
 	// #dxr_todo: these barriers should ideally be inserted by the high level code to allow more overlapped execution
-	AddUAVBarrier();
+	AddGlobalBarrier(ED3D12Access::BVHWrite, ED3D12Access::BVHRead);
 
 	for (const FRayTracingSceneBuildParams& SceneBuildParams : Params)
 	{
@@ -5506,8 +5525,18 @@ static void DispatchRays(FD3D12CommandContext& CommandContext,
 		DispatchRaysDescBuffer = Device->GetRayTracingDispatchRaysDescBuffer(QueueType);
 		FD3D12Resource* DispatchRaysDescBufferResource = DispatchRaysDescBuffer->GetResource();
 
-		CommandContext.TransitionResource(DispatchRaysDescBufferResource, D3D12_RESOURCE_STATE_INDIRECT_ARGUMENT, D3D12_RESOURCE_STATE_COPY_DEST, 0);
-		CommandContext.TransitionResource(ArgumentBuffer->GetResource(), D3D12_RESOURCE_STATE_NON_PIXEL_SHADER_RESOURCE|D3D12_RESOURCE_STATE_INDIRECT_ARGUMENT, D3D12_RESOURCE_STATE_COPY_SOURCE, 0);
+		CommandContext.AddBarrier(
+			DispatchRaysDescBufferResource,
+			ED3D12Access::IndirectArgs,
+			ED3D12Access::CopyDest,
+			0);
+
+		CommandContext.AddBarrier(
+			ArgumentBuffer->GetResource(),
+			ED3D12Access::UAVCompute | ED3D12Access::IndirectArgs,
+			ED3D12Access::CopySrc,
+			0);
+
 		CommandContext.FlushResourceBarriers();
 
 		// Compute the allocation & copy sizes
@@ -5541,18 +5570,17 @@ static void DispatchRays(FD3D12CommandContext& CommandContext,
 			IndirectDimensionSize
 		);
 
-		CommandContext.TransitionResource(
+		CommandContext.AddBarrier(
 			DispatchRaysDescBufferResource,
-			D3D12_RESOURCE_STATE_COPY_DEST,
-			D3D12_RESOURCE_STATE_INDIRECT_ARGUMENT,
-			0
-		);
+			ED3D12Access::CopyDest,
+			ED3D12Access::IndirectArgs,
+			0);
 
-		CommandContext.TransitionResource(ArgumentBuffer->GetResource(), 
-			D3D12_RESOURCE_STATE_COPY_SOURCE,
-			D3D12_RESOURCE_STATE_NON_PIXEL_SHADER_RESOURCE | D3D12_RESOURCE_STATE_INDIRECT_ARGUMENT, 
-			0
-		);
+		CommandContext.AddBarrier(
+			ArgumentBuffer->GetResource(),
+			ED3D12Access::CopySrc,
+			ED3D12Access::UAVCompute | ED3D12Access::IndirectArgs,
+			0);
 
 		CommandContext.FlushResourceBarriers();
 	}
@@ -6378,7 +6406,8 @@ void FD3D12CommandContext::RHIExecuteMultiIndirectClusterOperation(const FRayTra
 	check(Result == NVAPI_OK);
 
 	// Wait for the cluster ops to finish
-	AddUAVBarrier();
+	AddGlobalBarrier(ED3D12Access::UAVCompute, ED3D12Access::UAVCompute);
+
 #else
 	checkNoEntry();
 #endif // WITH_NVAPI
diff --git a/Engine/Source/Runtime/D3D12RHI/Private/D3D12RenderTarget.cpp b/Engine/Source/Runtime/D3D12RHI/Private/D3D12RenderTarget.cpp
--- a/Engine/Source/Runtime/D3D12RHI/Private/D3D12RenderTarget.cpp
+++ b/Engine/Source/Runtime/D3D12RHI/Private/D3D12RenderTarget.cpp
@@ -356,13 +356,23 @@ void FD3D12CommandContext::ResolveTexture(UE::RHICore::FResolveTextureInfo Info)
 			int32 DestSubresource   = CalcSubresource(Info.MipLevel, ArraySlice, DestDesc.NumMips);
 			int32 SourceSubresource = CalcSubresource(Info.MipLevel, ArraySlice, SourceDesc.NumMips);
 
+			// TODO: Is this taken care of now that we're paying attention to the RHI states?
+
 			// Source resource is in RTV, transition to D3D12_RESOURCE_STATE_RESOLVE_SOURCE and go back to D3D12_RESOURCE_STATE_RENDER_TARGET
 			// TODO: high level setting to D3D12_RESOURCE_STATE_RESOLVE_SOURCE
-			FScopedResourceBarrier ConditionalScopeResourceBarrierSrc(*this, SourceTexture->GetResource(), &SourceTexture->ResourceLocation, D3D12_RESOURCE_STATE_RESOLVE_SOURCE, SourceSubresource, D3D12_RESOURCE_STATE_RENDER_TARGET);
+			FScopedResourceBarrier ConditionalScopeResourceBarrierSrc(
+				*this,
+				SourceTexture->GetResource(),
+				ED3D12Access::RTV,
+				ED3D12Access::ResolveSrc,
+				SourceSubresource);
 
 			// Dest resource is already in D3D12_RESOURCE_STATE_RESOLVE_DEST, leave the FScopedResourceBarrier in case of using single state resources
-			FScopedResourceBarrier ConditionalScopeResourceBarrierDst(*this, DestTexture->GetResource(), &DestTexture->ResourceLocation, D3D12_RESOURCE_STATE_RESOLVE_DEST, DestSubresource);
-
+			FScopedResourceBarrier ConditionalScopeResourceBarrierDst(
+				*this,
+				DestTexture->GetResource(),
+				ED3D12Access::ResolveDst,
+				DestSubresource);
 
 			FlushResourceBarriers();
 			GraphicsCommandList()->ResolveSubresource(DestResource->GetResource(), DestSubresource, SourceTexture->GetResource()->GetResource(), SourceSubresource, DestFormatTypeless);
@@ -718,7 +728,15 @@ void FD3D12DynamicRHI::ReadSurfaceDataMSAARaw(FRHITexture* TextureRHI, FIntRect
 	TRefCountPtr<FD3D12Resource> NonMSAATexture2D;
 
 	const D3D12_HEAP_PROPERTIES HeapProps = CD3DX12_HEAP_PROPERTIES(D3D12_HEAP_TYPE_DEFAULT, NodeMask.GetNative(), NodeMask.GetNative());
-	VERIFYD3D12RESULT(Adapter->CreateCommittedResource(NonMSAADesc, NodeMask, HeapProps, D3D12_RESOURCE_STATE_RENDER_TARGET, nullptr, NonMSAATexture2D.GetInitReference(), nullptr));
+	VERIFYD3D12RESULT(
+		Adapter->CreateCommittedResource(
+			NonMSAADesc,
+			NodeMask,
+			HeapProps,
+			ED3D12Access::RTV,
+			nullptr,
+			NonMSAATexture2D.GetInitReference(),
+			nullptr));
 
 	FD3D12ResourceLocation ResourceLocation(Device);
 	ResourceLocation.AsStandAlone(NonMSAATexture2D);
@@ -795,7 +813,11 @@ void FD3D12DynamicRHI::ReadSurfaceDataMSAARaw(FRHITexture* TextureRHI, FIntRect
 				SampleIndex
 				);
 
-			FScopedResourceBarrier ScopeResourceBarrierSource(Context, NonMSAATexture2D, nullptr, D3D12_RESOURCE_STATE_COPY_SOURCE, SourceCopyLocation.SubresourceIndex);
+			FScopedResourceBarrier ScopeResourceBarrierSource(
+				Context,
+				NonMSAATexture2D,
+				ED3D12Access::CopySrc,
+				SourceCopyLocation.SubresourceIndex);
 			// Upload heap doesn't need to transition
 
 			// Copy the resolved sample data to the staging texture.
@@ -974,7 +996,11 @@ void FD3D12DynamicRHI::RHIReadSurfaceFloatData(FRHITexture* TextureRHI, FIntRect
 	{
 		FD3D12CommandContext& Context = FD3D12CommandContext::Get(ExecutingCmdList, GPUIndex);
 
-		FScopedResourceBarrier ConditionalScopeResourceBarrier(Context, Texture->GetResource(), &Texture->ResourceLocation, D3D12_RESOURCE_STATE_COPY_SOURCE, SourceCopyLocation.SubresourceIndex);
+		FScopedResourceBarrier ConditionalScopeResourceBarrier(
+			Context,
+			Texture->GetResource(),
+			ED3D12Access::CopySrc,
+			SourceCopyLocation.SubresourceIndex);
 		// Don't need to transition upload heaps
 
 		Context.FlushResourceBarriers();
@@ -1089,7 +1115,11 @@ void FD3D12DynamicRHI::RHIRead3DSurfaceFloatData(FRHITexture* TextureRHI, FIntRe
 	{
 		FD3D12CommandContext& Context = FD3D12CommandContext::Get(ExecutingCmdList, GPUIndex);
 
-		FScopedResourceBarrier ConditionalScopeResourceBarrier(Context, Texture->GetResource(), &Texture->ResourceLocation, D3D12_RESOURCE_STATE_COPY_SOURCE, SourceCopyLocation.SubresourceIndex);
+		FScopedResourceBarrier ConditionalScopeResourceBarrier(
+			Context,
+			Texture->GetResource(),
+			ED3D12Access::CopySrc,
+			SourceCopyLocation.SubresourceIndex);
 		// Don't need to transition upload heaps
 
 		Context.FlushResourceBarriers();
diff --git a/Engine/Source/Runtime/D3D12RHI/Private/D3D12Resources.cpp b/Engine/Source/Runtime/D3D12RHI/Private/D3D12Resources.cpp
--- a/Engine/Source/Runtime/D3D12RHI/Private/D3D12Resources.cpp
+++ b/Engine/Source/Runtime/D3D12RHI/Private/D3D12Resources.cpp
@@ -26,8 +26,17 @@ static TAutoConsoleVariable<int32> CVarD3D12ReservedResourceHeapSizeMB(
 /////////////////////////////////////////////////////////////////////
 
 
-void ID3D12ResourceAllocator::AllocateTexture(uint32 GPUIndex, D3D12_HEAP_TYPE InHeapType, const FD3D12ResourceDesc& InDesc, EPixelFormat InUEFormat, ED3D12ResourceStateMode InResourceStateMode,
-	D3D12_RESOURCE_STATES InCreateState, const D3D12_CLEAR_VALUE* InClearValue, const TCHAR* InName, FD3D12ResourceLocation& ResourceLocation)
+void ID3D12ResourceAllocator::AllocateTexture(
+	uint32 GPUIndex,
+	D3D12_HEAP_TYPE InHeapType,
+	const FD3D12ResourceDesc& InDesc,
+	EPixelFormat InUEFormat,
+	ED3D12Access InInitialD3D12Access,
+	ED3D12ResourceStateMode InResourceStateMode,
+	ED3D12Access InDefaultD3D12Access,
+	const D3D12_CLEAR_VALUE* InClearValue,
+	const TCHAR* InName,
+	FD3D12ResourceLocation& ResourceLocation)
 {
 	// Check if texture can be 4K aligned
 	FD3D12ResourceDesc Desc = InDesc;
@@ -36,7 +45,18 @@ void ID3D12ResourceAllocator::AllocateTexture(uint32 GPUIndex, D3D12_HEAP_TYPE I
 
 	// Get the size and alignment for the allocation
 	D3D12_RESOURCE_ALLOCATION_INFO Info = FD3D12DynamicRHI::GetD3DRHI()->GetAdapter().GetDevice(GPUIndex)->GetResourceAllocationInfo(Desc);
-	AllocateResource(GPUIndex, InHeapType, Desc, Info.SizeInBytes, Info.Alignment, InResourceStateMode, InCreateState, InClearValue, InName, ResourceLocation);
+	AllocateResource(
+		GPUIndex,
+		InHeapType,
+		Desc,
+		Info.SizeInBytes,
+		Info.Alignment,
+		InInitialD3D12Access,
+		InResourceStateMode,
+		InDefaultD3D12Access,
+		InClearValue,
+		InName,
+		ResourceLocation);
 }
 
 #if D3D12RHI_SUPPORTS_UNCOMPRESSED_UAV
@@ -44,12 +64,12 @@ TArray<DXGI_FORMAT, TInlineAllocator<4>> FD3D12ResourceDesc::GetCastableFormats(
 {
 	TArray<DXGI_FORMAT, TInlineAllocator<4>> Result;
 
-	// We have to add the 'implied' castable formats for SRVs. Since we don't have any sRGB flags here, just add both formats.
-	Result.Emplace(UE::DXGIUtilities::FindShaderResourceFormat(Format, true));
-	Result.Emplace(UE::DXGIUtilities::FindShaderResourceFormat(Format, false));
-
 	if (UAVPixelFormat != PF_Unknown)
 	{
+		// We have to add the 'implied' castable formats for SRVs. Since we don't have any sRGB flags here, just add both formats.
+		Result.Emplace(UE::DXGIUtilities::FindShaderResourceFormat(Format, true));
+		Result.Emplace(UE::DXGIUtilities::FindShaderResourceFormat(Format, false));
+
 		// Add the uncompressed UAV format we want
 		Result.Emplace((DXGI_FORMAT)GPixelFormats[UAVPixelFormat].PlatformFormat);
 	}
@@ -70,21 +90,30 @@ int64 FD3D12Resource::NoStateTrackingResourceCount = 0;
 FD3D12Resource::FD3D12Resource(FD3D12Device* ParentDevice,
 	FRHIGPUMask VisibleNodes,
 	ID3D12Resource* InResource,
-	D3D12_RESOURCE_STATES InInitialState,
+	ED3D12Access InInitialD3D12Access,
 	const FD3D12ResourceDesc& InDesc,
 	FD3D12Heap* InHeap,
 	D3D12_HEAP_TYPE InHeapType) : 
-	FD3D12Resource(ParentDevice, VisibleNodes, InResource, InInitialState, ED3D12ResourceStateMode::Default, D3D12_RESOURCE_STATE_TBD, InDesc, InHeap, InHeapType)
+	FD3D12Resource(
+		ParentDevice,
+		VisibleNodes,
+		InResource,
+		InInitialD3D12Access,
+		ED3D12ResourceStateMode::Default,
+		ED3D12Access::Unknown,
+		InDesc,
+		InHeap,
+		InHeapType)
 {
 }
 
 FD3D12Resource::FD3D12Resource(FD3D12Device* ParentDevice,
 	FRHIGPUMask VisibleNodes,
 	ID3D12Resource* InResource,
-	D3D12_RESOURCE_STATES InInitialState,
+	ED3D12Access InInitialD3D12Access,
 	ED3D12ResourceStateMode InResourceStateMode,
-	D3D12_RESOURCE_STATES InDefaultResourceState,
-	const FD3D12ResourceDesc& InDesc,
+	ED3D12Access InDefaultD3D12Access,
+	FD3D12ResourceDesc const& InDesc,
 	FD3D12Heap* InHeap,
 	D3D12_HEAP_TYPE InHeapType)
 	: FD3D12DeviceChild(ParentDevice)
@@ -129,7 +158,11 @@ FD3D12Resource::FD3D12Resource(FD3D12Device* ParentDevice,
 		GPUVirtualAddress = Resource->GetGPUVirtualAddress();
 	}
 
-	InitalizeResourceState(InInitialState, InResourceStateMode, InDefaultResourceState);
+	InitializeResourceState(
+		HeapProps,
+		InInitialD3D12Access,
+		InResourceStateMode,
+		InDefaultD3D12Access);
 
 #if NV_AFTERMATH
 	AftermathHandle = UE::RHICore::Nvidia::Aftermath::D3D12::RegisterResource(InResource);
@@ -736,8 +769,29 @@ void FD3D12Heap::BeginTrackingResidency(uint64 Size)
 //	FD3D12 Adapter
 /////////////////////////////////////////////////////////////////////
 
-HRESULT FD3D12Adapter::CreateCommittedResource(const FD3D12ResourceDesc& InDesc, FRHIGPUMask CreationNode, const D3D12_HEAP_PROPERTIES& HeapProps, D3D12_RESOURCE_STATES InInitialState,
-	ED3D12ResourceStateMode InResourceStateMode, D3D12_RESOURCE_STATES InDefaultState, const D3D12_CLEAR_VALUE* ClearValue, FD3D12Resource** ppOutResource, const TCHAR* Name, bool bVerifyHResult)
+static D3D12_RESOURCE_FLAGS GetRayTracingResourceFlags(ED3D12Access InDefaultD3D12Access)
+{
+#if D3D12_RHI_RAYTRACING
+	if (EnumHasAnyFlags(InDefaultD3D12Access, ED3D12Access::BVHRead | ED3D12Access::BVHWrite))
+	{
+		return D3D12_RESOURCE_FLAG_ALLOW_UNORDERED_ACCESS | 
+		       D3D12_RESOURCE_FLAG_RAYTRACING_ACCELERATION_STRUCTURE;
+	}
+#endif // D3D12_RHI_RAYTRACING
+	return D3D12_RESOURCE_FLAG_NONE;
+}
+
+HRESULT FD3D12Adapter::CreateCommittedResource(
+	const FD3D12ResourceDesc& InDesc,
+	FRHIGPUMask CreationNode,
+	const D3D12_HEAP_PROPERTIES& HeapProps,
+	ED3D12Access InInitialD3D12Access,
+	ED3D12ResourceStateMode InResourceStateMode,
+	ED3D12Access InDefaultD3D12Access,
+	const D3D12_CLEAR_VALUE* ClearValue,
+	FD3D12Resource** ppOutResource,
+	const TCHAR* Name,
+	bool bVerifyHResult)
 {
 	if (!ppOutResource)
 	{
@@ -763,64 +817,35 @@ HRESULT FD3D12Adapter::CreateCommittedResource(const FD3D12ResourceDesc& InDesc,
 		}
 	}
 
-#if D3D12_RHI_RAYTRACING
-	if (InDefaultState == D3D12_RESOURCE_STATE_RAYTRACING_ACCELERATION_STRUCTURE)
-	{
-		LocalDesc.Flags |= D3D12_RESOURCE_FLAG_ALLOW_UNORDERED_ACCESS;
-	}
-#endif // D3D12_RHI_RAYTRACING
+	LocalDesc.Flags |= GetRayTracingResourceFlags(InDefaultD3D12Access);
 
 #if D3D12_WITH_CUSTOM_TEXTURE_LAYOUT
 	ApplyCustomTextureLayout(LocalDesc, *this);
 #endif
 
-	HRESULT hr = S_OK;
-#if INTEL_EXTENSIONS
-	if (InDesc.bRequires64BitAtomicSupport && IsRHIDeviceIntel() && GDX12INTCAtomicUInt64Emulation)
-	{
-		INTC_D3D12_RESOURCE_DESC_0001 IntelLocalDesc{};
-		IntelLocalDesc.pD3D12Desc = &LocalDesc;
-		IntelLocalDesc.EmulatedTyped64bitAtomics = true;
-
-		hr = INTC_D3D12_CreateCommittedResource(FD3D12DynamicRHI::GetD3DRHI()->GetIntelExtensionContext(), &HeapProps, HeapFlags, &IntelLocalDesc, InInitialState, ClearValue, IID_PPV_ARGS(pResource.GetInitReference()));
-	}
-	else
-#endif
-#if D3D12RHI_SUPPORTS_UNCOMPRESSED_UAV
-	if (InDesc.SupportsUncompressedUAV())
-	{
-		// Convert the desc to the version required by CreateCommittedResource3
-		const CD3DX12_RESOURCE_DESC1 LocalDesc1(LocalDesc);
-
-		// Common layout is the required starting state for any "legacy" transitions
-		const D3D12_BARRIER_LAYOUT InitialLayout = D3D12_BARRIER_LAYOUT_COMMON;
-
-		checkf(InInitialState == D3D12_RESOURCE_STATE_COMMON, TEXT("RESOURCE_STATE_COMMON is required for castable resources (Given: %d)"), InInitialState);
-
-		ID3D12ProtectedResourceSession* ProtectedSession = nullptr;
-
-		const TArray<DXGI_FORMAT, TInlineAllocator<4>> CastableFormats = InDesc.GetCastableFormats();
-
-		hr = RootDevice12->CreateCommittedResource3(
-			&HeapProps,
+	HRESULT hr =
+		Barriers->CreateCommittedResource(
+			*this,
+			HeapProps, 
 			HeapFlags,
-			&LocalDesc1,
-			InitialLayout,
+			LocalDesc,
+			InInitialD3D12Access,
 			ClearValue,
-			ProtectedSession,
-			CastableFormats.Num(),
-			CastableFormats.GetData(),
-			IID_PPV_ARGS(pResource.GetInitReference()));
-	}
-	else
-#endif
-	{
-		hr = RootDevice->CreateCommittedResource(&HeapProps, HeapFlags, &LocalDesc, InInitialState, ClearValue, IID_PPV_ARGS(pResource.GetInitReference()));
-	}
+			pResource);
+
 	if (SUCCEEDED(hr))
 	{
 		// Set the output pointer
-		*ppOutResource = new FD3D12Resource(GetDevice(CreationNode.ToIndex()), CreationNode, pResource, InInitialState, InResourceStateMode, InDefaultState, LocalDesc, nullptr, HeapProps.Type);
+		*ppOutResource = new FD3D12Resource(
+			GetDevice(CreationNode.ToIndex()),
+			CreationNode,
+			pResource,
+			InInitialD3D12Access,
+			InResourceStateMode,
+			InDefaultD3D12Access,
+			LocalDesc,
+			nullptr,
+			HeapProps.Type);
 		(*ppOutResource)->AddRef();
 
 		// Set a default name (can override later).
@@ -845,8 +870,16 @@ HRESULT FD3D12Adapter::CreateCommittedResource(const FD3D12ResourceDesc& InDesc,
 	return hr;
 }
 
-HRESULT FD3D12Adapter::CreateReservedResource(const FD3D12ResourceDesc& InDesc, FRHIGPUMask CreationNode, D3D12_RESOURCE_STATES InInitialState,
-	ED3D12ResourceStateMode InResourceStateMode, D3D12_RESOURCE_STATES InDefaultState, const D3D12_CLEAR_VALUE* ClearValue, FD3D12Resource** ppOutResource, const TCHAR* Name, bool bVerifyHResult)
+HRESULT FD3D12Adapter::CreateReservedResource(
+	const FD3D12ResourceDesc& InDesc,
+	FRHIGPUMask CreationNode,
+	ED3D12Access InInitialD3D12Access,
+	ED3D12ResourceStateMode InResourceStateMode,
+	ED3D12Access InDefaultD3D12Access,
+	const D3D12_CLEAR_VALUE* ClearValue,
+	FD3D12Resource** ppOutResource,
+	const TCHAR* Name,
+	bool bVerifyHResult)
 {
 	if (!ppOutResource)
 	{
@@ -877,45 +910,29 @@ HRESULT FD3D12Adapter::CreateReservedResource(const FD3D12ResourceDesc& InDesc,
 		TEXT("Reserved resources must use either 64KB alignment or 0 (unspecified/default), but have %d. See FD3D12DynamicRHI::GetResourceDesc()."),
 			uint32(LocalDesc.Alignment));
 
-#if D3D12_RHI_RAYTRACING
-	if (InDefaultState == D3D12_RESOURCE_STATE_RAYTRACING_ACCELERATION_STRUCTURE)
-	{
-		LocalDesc.Flags |= D3D12_RESOURCE_FLAG_ALLOW_UNORDERED_ACCESS;
-	}
-#endif // D3D12_RHI_RAYTRACING
-
-	HRESULT hr = S_OK;
-	
-#if D3D12RHI_SUPPORTS_UNCOMPRESSED_UAV
-	if (InDesc.SupportsUncompressedUAV())
-	{
-		checkf(InInitialState == D3D12_RESOURCE_STATE_COMMON, TEXT("RESOURCE_STATE_COMMON is required for castable resources (Given: %d)"), InInitialState);
-
-		// Common layout is the require starting state for any "legacy" transitions
-		const D3D12_BARRIER_LAYOUT InitialLayout = D3D12_BARRIER_LAYOUT_COMMON;
+	LocalDesc.Flags |= GetRayTracingResourceFlags(InDefaultD3D12Access);
 
-		ID3D12ProtectedResourceSession* ProtectedSession = nullptr;
-		const TArray<DXGI_FORMAT, TInlineAllocator<4>> CastableFormats = InDesc.GetCastableFormats();
-
-		hr = RootDevice12->CreateReservedResource2(
-			&LocalDesc,
-			InitialLayout,
+	HRESULT hr = 
+		Barriers->CreateReservedResource(
+			*this,
+			LocalDesc, 
+			InInitialD3D12Access,
 			ClearValue,
-			ProtectedSession,
-			CastableFormats.Num(),
-			CastableFormats.GetData(),
-			IID_PPV_ARGS(pResource.GetInitReference()));
-	}
-	else
-#endif
-	{
-		hr = RootDevice->CreateReservedResource(&LocalDesc, InInitialState, ClearValue, IID_PPV_ARGS(pResource.GetInitReference()));
-	}
+			pResource);
 
 	if (SUCCEEDED(hr))
 	{
 		// Set the output pointer
-		*ppOutResource = new FD3D12Resource(GetDevice(CreationNode.ToIndex()), CreationNode, pResource, InInitialState, InResourceStateMode, InDefaultState, LocalDesc, nullptr /*Heap*/, D3D12_HEAP_TYPE_DEFAULT);
+		*ppOutResource = new FD3D12Resource(
+			GetDevice(CreationNode.ToIndex()),
+			CreationNode,
+			pResource,
+			InInitialD3D12Access,
+			InResourceStateMode,
+			InDefaultD3D12Access,
+			LocalDesc,
+			nullptr /*Heap*/,
+			D3D12_HEAP_TYPE_DEFAULT);
 		(*ppOutResource)->AddRef();
 
 		// Set a default name (can override later).
@@ -937,8 +954,17 @@ HRESULT FD3D12Adapter::CreateReservedResource(const FD3D12ResourceDesc& InDesc,
 	return hr;
 }
 
-HRESULT FD3D12Adapter::CreatePlacedResource(const FD3D12ResourceDesc& InDesc, FD3D12Heap* BackingHeap, uint64 HeapOffset, D3D12_RESOURCE_STATES InInitialState, ED3D12ResourceStateMode InResourceStateMode,
-	D3D12_RESOURCE_STATES InDefaultState, const D3D12_CLEAR_VALUE* ClearValue, FD3D12Resource** ppOutResource, const TCHAR* Name, bool bVerifyHResult)
+HRESULT FD3D12Adapter::CreatePlacedResource(
+	const FD3D12ResourceDesc& InDesc,
+	FD3D12Heap* BackingHeap,
+	uint64 HeapOffset,
+	ED3D12Access InInitialD3D12Access,
+	ED3D12ResourceStateMode InResourceStateMode,
+	ED3D12Access InDefaultD3D12Access,
+	const D3D12_CLEAR_VALUE* ClearValue,
+	FD3D12Resource** ppOutResource,
+	const TCHAR* Name,
+	bool bVerifyHResult)
 {
 	if (!ppOutResource)
 	{
@@ -948,47 +974,16 @@ HRESULT FD3D12Adapter::CreatePlacedResource(const FD3D12ResourceDesc& InDesc, FD
 	ID3D12Heap* Heap = BackingHeap->GetHeap();
 
 	TRefCountPtr<ID3D12Resource> pResource;
-	HRESULT hr = S_OK;
-#if INTEL_EXTENSIONS
-	if (InDesc.bRequires64BitAtomicSupport && IsRHIDeviceIntel() && GDX12INTCAtomicUInt64Emulation)
-	{
-		FD3D12ResourceDesc LocalDesc = InDesc;
-		INTC_D3D12_RESOURCE_DESC_0001 IntelLocalDesc{};
-		IntelLocalDesc.pD3D12Desc = &LocalDesc;
-		IntelLocalDesc.EmulatedTyped64bitAtomics = true;
-
-		hr = INTC_D3D12_CreatePlacedResource(FD3D12DynamicRHI::GetD3DRHI()->GetIntelExtensionContext(), Heap, HeapOffset, &IntelLocalDesc, InInitialState, ClearValue, IID_PPV_ARGS(pResource.GetInitReference()));
-	}
-	else
-#endif
-#if D3D12RHI_SUPPORTS_UNCOMPRESSED_UAV
-	if (InDesc.SupportsUncompressedUAV())
-	{
-		checkf(InInitialState == D3D12_RESOURCE_STATE_COMMON, TEXT("RESOURCE_STATE_COMMON is required for castable resources (Given: %d)"), InInitialState);
-
-		// Convert the desc to the version required by CreatePlacedResource2
-		const CD3DX12_RESOURCE_DESC1 LocalDesc(InDesc);
-
-		// Common layout is the required starting state for any "legacy" transitions
-		const D3D12_BARRIER_LAYOUT InitialLayout = D3D12_BARRIER_LAYOUT_COMMON;
-		
-		const TArray<DXGI_FORMAT, TInlineAllocator<4>> CastableFormats = InDesc.GetCastableFormats();
-
-		hr = RootDevice10->CreatePlacedResource2(
+	HRESULT hr =
+		Barriers->CreatePlacedResource(
+			*this,
 			Heap,
 			HeapOffset,
-			&LocalDesc,
-			InitialLayout,
+			InDesc,
+			InInitialD3D12Access,
 			ClearValue,
-			CastableFormats.Num(),
-			CastableFormats.GetData(),
-			IID_PPV_ARGS(pResource.GetInitReference()));
-	}
-	else
-#endif
-	{
-		hr = RootDevice->CreatePlacedResource(Heap, HeapOffset, &InDesc, InInitialState, ClearValue, IID_PPV_ARGS(pResource.GetInitReference()));
-	}
+			pResource);
+
 	if (SUCCEEDED(hr))
 	{
 		FD3D12Device* Device = BackingHeap->GetParentDevice();
@@ -998,9 +993,9 @@ HRESULT FD3D12Adapter::CreatePlacedResource(const FD3D12ResourceDesc& InDesc, FD
 		*ppOutResource = new FD3D12Resource(Device,
 			Device->GetVisibilityMask(),
 			pResource,
-			InInitialState,
+			InInitialD3D12Access,
 			InResourceStateMode,
-			InDefaultState,
+			InDefaultD3D12Access,
 			InDesc,
 			BackingHeap,
 			HeapDesc.Properties.Type);
@@ -1063,24 +1058,60 @@ void FD3D12Adapter::TraceMemoryAllocation(FD3D12Resource* Resource)
 #endif
 }
 
-HRESULT FD3D12Adapter::CreateBuffer(D3D12_HEAP_TYPE HeapType, FRHIGPUMask CreationNode, FRHIGPUMask VisibleNodes, uint64 HeapSize, FD3D12Resource** ppOutResource, const TCHAR* Name, D3D12_RESOURCE_FLAGS Flags)
+HRESULT FD3D12Adapter::CreateBuffer(
+	D3D12_HEAP_TYPE HeapType,
+	FRHIGPUMask CreationNode,
+	FRHIGPUMask VisibleNodes,
+	uint64 HeapSize,
+	FD3D12Resource** ppOutResource,
+	const TCHAR* Name,
+	D3D12_RESOURCE_FLAGS Flags)
 {
 	const D3D12_HEAP_PROPERTIES HeapProps = CD3DX12_HEAP_PROPERTIES(HeapType, CreationNode.GetNative(), VisibleNodes.GetNative());
-	const D3D12_RESOURCE_STATES InitialState = DetermineInitialResourceState(HeapProps.Type, &HeapProps);
-	return CreateBuffer(HeapProps, CreationNode, InitialState, ED3D12ResourceStateMode::Default, D3D12_RESOURCE_STATE_TBD, HeapSize, ppOutResource, Name, Flags);
+	const ED3D12Access InitialD3D12Access = DetermineInitialBufferD3D12Access(HeapProps.Type, &HeapProps);
+	return CreateBuffer(
+		HeapProps,
+		CreationNode,
+		InitialD3D12Access,
+		ED3D12ResourceStateMode::Default,
+		ED3D12Access::Unknown,
+		HeapSize,
+		ppOutResource,
+		Name,
+		Flags);
 }
 
-HRESULT FD3D12Adapter::CreateBuffer(D3D12_HEAP_TYPE HeapType, FRHIGPUMask CreationNode, FRHIGPUMask VisibleNodes, D3D12_RESOURCE_STATES InitialState, ED3D12ResourceStateMode ResourceStateMode, uint64 HeapSize, FD3D12Resource** ppOutResource, const TCHAR* Name, D3D12_RESOURCE_FLAGS Flags)
+HRESULT FD3D12Adapter::CreateBuffer(
+	D3D12_HEAP_TYPE HeapType,
+	FRHIGPUMask CreationNode,
+	FRHIGPUMask VisibleNodes,
+	ED3D12Access InInitialD3D12Access,
+	ED3D12ResourceStateMode ResourceStateMode,
+	ED3D12Access InDefaultD3D12Access,
+	uint64 HeapSize,
+	FD3D12Resource** ppOutResource,
+	const TCHAR* Name,
+	D3D12_RESOURCE_FLAGS Flags)
 {
 	const D3D12_HEAP_PROPERTIES HeapProps = CD3DX12_HEAP_PROPERTIES(HeapType, CreationNode.GetNative(), VisibleNodes.GetNative());
-	return CreateBuffer(HeapProps, CreationNode, InitialState, ResourceStateMode, InitialState, HeapSize, ppOutResource, Name, Flags);
+	return CreateBuffer(
+		HeapProps,
+		CreationNode,
+		InInitialD3D12Access,
+		ResourceStateMode,
+		InDefaultD3D12Access,
+		HeapSize,
+		ppOutResource,
+		Name,
+		Flags);
 }
 
-HRESULT FD3D12Adapter::CreateBuffer(const D3D12_HEAP_PROPERTIES& HeapProps,
+HRESULT FD3D12Adapter::CreateBuffer(
+	const D3D12_HEAP_PROPERTIES& HeapProps,
 	FRHIGPUMask CreationNode,
-	D3D12_RESOURCE_STATES InitialState,
+	ED3D12Access InInitialD3D12Access,
 	ED3D12ResourceStateMode ResourceStateMode,
-	D3D12_RESOURCE_STATES InDefaultState,
+	ED3D12Access InDefaultD3D12Access,
 	uint64 HeapSize,
 	FD3D12Resource** ppOutResource,
 	const TCHAR* Name,
@@ -1092,14 +1123,16 @@ HRESULT FD3D12Adapter::CreateBuffer(const D3D12_HEAP_PROPERTIES& HeapProps,
 	}
 
 	const D3D12_RESOURCE_DESC BufDesc = CD3DX12_RESOURCE_DESC::Buffer(HeapSize, Flags);
-	return CreateCommittedResource(BufDesc,
+	return CreateCommittedResource(
+		BufDesc,
 		CreationNode,
 		HeapProps,
-		InitialState,
+		InInitialD3D12Access,
 		ResourceStateMode,
-		InDefaultState,
+		InDefaultD3D12Access,
 		nullptr,
-		ppOutResource, Name);
+		ppOutResource,
+		Name);
 }
 
 #if !D3D12RHI_SUPPORTS_UNCOMPRESSED_UAV
@@ -1108,7 +1141,6 @@ void FD3D12Adapter::CreateUAVAliasResourceDesc(const FD3D12ResourceLocation& Loc
 	FD3D12Resource* SourceResource = Location.GetResource();
 
 	const FD3D12ResourceDesc& SourceDesc = SourceResource->GetDesc();
-
 	const EPixelFormat SourceFormat = SourceDesc.PixelFormat;
 	const EPixelFormat AliasTextureFormat = SourceDesc.UAVPixelFormat;
 
@@ -1451,7 +1483,10 @@ void FD3D12ResourceLocation::AsStandAlone(FD3D12Resource* Resource, uint64 InSiz
 }
 
 
-bool FD3D12ResourceLocation::OnAllocationMoved(FD3D12ContextArray const& Contexts, FRHIPoolAllocationData* InNewData, D3D12_RESOURCE_STATES& OutCreateState)
+bool FD3D12ResourceLocation::OnAllocationMoved(
+	FD3D12ContextArray const& Contexts,
+	FRHIPoolAllocationData* InNewData,
+	ED3D12Access& OutD3D12Access)
 {
 	// Assume linked list allocated for now - only defragging allocator
 	FRHIPoolAllocationData& AllocationData = GetPoolAllocatorPrivateData().PoolData;
@@ -1503,18 +1538,18 @@ bool FD3D12ResourceLocation::OnAllocationMoved(FD3D12ContextArray const& Context
 		// recreate the placed resource (ownership of current resource is already handled during the internal move)
 		FD3D12HeapAndOffset HeapAndOffset = NewAllocator->GetBackingHeapAndAllocationOffsetInBytes(*this);
 
-		D3D12_RESOURCE_STATES CreateState;
+		OutD3D12Access = ED3D12Access::Unknown;
 		ED3D12ResourceStateMode ResourceStateMode;
 		if (CurrentResource->RequiresResourceStateTracking())
 		{
 			// The newly created placed resource will be copied into by the defragger. Create it in the COPY_DEST to avoid an additional transition.
 			// Standard resource state tracking will handle transitioning the resource out of this state as required.
-			CreateState = D3D12_RESOURCE_STATE_COPY_DEST;
+			OutD3D12Access = ED3D12Access::CopyDest;
 			ResourceStateMode = ED3D12ResourceStateMode::MultiState;
 		}
 		else
 		{
-			CreateState = CurrentResource->GetDefaultResourceState();
+			OutD3D12Access = CurrentResource->GetDefaultAccess();
 			ResourceStateMode = ED3D12ResourceStateMode::Default;
 		}
 
@@ -1524,11 +1559,21 @@ bool FD3D12ResourceLocation::OnAllocationMoved(FD3D12ContextArray const& Context
 		FName Name = CurrentResource->GetName();
 
 		FD3D12Resource* NewResource = nullptr;
-		VERIFYD3D12RESULT(CurrentResource->GetParentDevice()->GetParentAdapter()->CreatePlacedResource(CurrentResource->GetDesc(), HeapAndOffset.Heap, HeapAndOffset.Offset, CreateState, ResourceStateMode, D3D12_RESOURCE_STATE_TBD, ClearValue, &NewResource, *Name.ToString()));
+		FD3D12Adapter* const Adapter = CurrentResource->GetParentDevice()->GetParentAdapter();
+		VERIFYD3D12RESULT(
+			Adapter->CreatePlacedResource(
+				CurrentResource->GetDesc(),
+				HeapAndOffset.Heap,
+				HeapAndOffset.Offset,
+				OutD3D12Access,
+				ResourceStateMode,
+				OutD3D12Access,
+				ClearValue,
+				&NewResource,
+				*Name.ToString()));
+
 		UnderlyingResource = NewResource;
 		GPUVirtualAddress = UnderlyingResource->GetGPUVirtualAddress() + OffsetFromBaseOfResource;
-		
-		OutCreateState = CreateState == D3D12_RESOURCE_STATE_TBD ? CurrentResource->GetWritableState() : CreateState;
 	}
 
 	// Refresh aliases
@@ -1544,25 +1589,25 @@ bool FD3D12ResourceLocation::OnAllocationMoved(FD3D12ContextArray const& Context
 	// Notify all the dependent resources about the change
 	Owner->ResourceRenamed(Contexts);
 
-	if (OutCreateState == D3D12_RESOURCE_STATE_TBD)
+	if (OutD3D12Access == ED3D12Access::Unknown)
 	{
-		ERHIAccess ResourceAccess = ERHIAccess::Unknown;
 		FD3D12Texture* CurrentResourceRHITexture = nullptr;
 
 		if (CurrentResource->GetDesc().Dimension == D3D12_RESOURCE_DIMENSION_BUFFER)
 		{
 			FD3D12Buffer* CurrentResourceRHIBuffer = (FD3D12Buffer*)Owner;
-			ResourceAccess = CurrentResourceRHIBuffer->GetTrackedAccess_Unsafe();
+			OutD3D12Access = ConvertToD3D12Access(CurrentResourceRHIBuffer->GetTrackedAccess_Unsafe());
 		}
 		else
 		{
 			CurrentResourceRHITexture = (FD3D12Texture*)Owner;
-			ResourceAccess = CurrentResourceRHITexture->GetTrackedAccess_Unsafe();
+			OutD3D12Access = ConvertToD3D12Access(CurrentResourceRHITexture->GetTrackedAccess_Unsafe());
 		}
-
-		OutCreateState = ResourceAccess == ERHIAccess::Unknown ? CurrentResource->GetInitialResourceState() : GetD3D12ResourceState(ResourceAccess, ED3D12QueueType::Direct, CurrentResourceRHITexture);
 	}
 
+	check(IsValidAccess(OutD3D12Access));
+	check(OutD3D12Access != ED3D12Access::Unknown);
+
 	return true;
 }
 
@@ -1582,130 +1627,6 @@ bool FD3D12ResourceLocation::IsStandaloneOrPooledPlacedResource() const
 	return bStandalone || bPoolPlacedResource;
 }
 
-/////////////////////////////////////////////////////////////////////
-//	FD3D12 Resource Barrier Batcher
-/////////////////////////////////////////////////////////////////////
-
-// Add a UAV barrier to the batch. Ignoring the actual resource for now.
-void FD3D12ResourceBarrierBatcher::AddUAV()
-{
-	FD3D12ResourceBarrier& Barrier = Barriers.Emplace_GetRef();
-	Barrier.Type = D3D12_RESOURCE_BARRIER_TYPE_UAV;
-	Barrier.Flags = D3D12_RESOURCE_BARRIER_FLAG_NONE;
-	Barrier.UAV.pResource = nullptr;	// Ignore the resource ptr for now. HW doesn't do anything with it.
-}
-
-// Add a transition resource barrier to the batch. Returns the number of barriers added, which may be negative if an existing barrier was cancelled.
-int32 FD3D12ResourceBarrierBatcher::AddTransition(FD3D12Resource* pResource, D3D12_RESOURCE_STATES Before, D3D12_RESOURCE_STATES After, uint32 Subresource)
-{
-	check(Before != After);
-
-	if (Barriers.Num())
-	{
-		// Check if we are simply reverting the last transition. In that case, we can just remove both transitions.
-		// This happens fairly frequently due to resource pooling since different RHI buffers can point to the same underlying D3D buffer.
-		// Instead of ping-ponging that underlying resource between COPY_DEST and GENERIC_READ, several copies can happen without a ResourceBarrier() in between.
-		// Doing this check also eliminates a D3D debug layer warning about multiple transitions of the same subresource.
-		const FD3D12ResourceBarrier& Last = Barriers.Last();
-
-		if (Last.Type == D3D12_RESOURCE_BARRIER_TYPE_TRANSITION
-			&& pResource->GetResource() == Last.Transition.pResource
-			&& Subresource == Last.Transition.Subresource
-			&& Before      == Last.Transition.StateAfter
-			&& After       == Last.Transition.StateBefore
-			)
-		{
-			Barriers.RemoveAt(Barriers.Num() - 1);
-			return -1;
-		}
-	}
-
-	check(IsValidD3D12ResourceState(Before) && IsValidD3D12ResourceState(After));
-
-	FD3D12ResourceBarrier& Barrier = Barriers.Emplace_GetRef(CD3DX12_RESOURCE_BARRIER::Transition(pResource->GetResource(), Before, After, Subresource));
-	if (pResource->IsBackBuffer() && EnumHasAnyFlags(After, BackBufferBarrierWriteTransitionTargets))
-	{
-		Barrier.Flags |= BarrierFlag_CountAsIdleTime;
-	}
-
-	return 1;
-}
-
-void FD3D12ResourceBarrierBatcher::AddAliasingBarrier(ID3D12Resource* InResourceBefore, ID3D12Resource* InResourceAfter)
-{
-	FD3D12ResourceBarrier& Barrier = Barriers.Emplace_GetRef();
-	Barrier.Type = D3D12_RESOURCE_BARRIER_TYPE_ALIASING;
-	Barrier.Flags = D3D12_RESOURCE_BARRIER_FLAG_NONE;
-	Barrier.Aliasing.pResourceBefore = InResourceBefore;
-	Barrier.Aliasing.pResourceAfter = InResourceAfter;
-}
-
-void FD3D12ResourceBarrierBatcher::FlushIntoCommandList(FD3D12CommandList& CommandList, FD3D12QueryAllocator& TimestampAllocator)
-{
-	auto InsertTimestamp = [&](bool bBegin)
-	{
-#if RHI_NEW_GPU_PROFILER
-		if (bBegin)
-		{
-			auto& Event = CommandList.EmplaceProfilerEvent<UE::RHI::GPUProfiler::FEvent::FEndWork>();
-			CommandList.EndQuery(TimestampAllocator.Allocate(ED3D12QueryType::ProfilerTimestampBOP, &Event.GPUTimestampBOP));
-		}
-		else
-		{
-			// CPUTimestamp is filled in at submission time in FlushProfilerEvents
-			auto& Event = CommandList.EmplaceProfilerEvent<UE::RHI::GPUProfiler::FEvent::FBeginWork>(0);
-			CommandList.EndQuery(TimestampAllocator.Allocate(ED3D12QueryType::ProfilerTimestampTOP, &Event.GPUTimestampTOP));
-		}
-#else
-		ED3D12QueryType Type = bBegin
-			? ED3D12QueryType::IdleBegin
-			: ED3D12QueryType::IdleEnd;
-
-		CommandList.EndQuery(TimestampAllocator.Allocate(Type, nullptr));
-#endif
-	};
-
-	for (int32 BatchStart = 0, BatchEnd = 0; BatchStart < Barriers.Num(); BatchStart = BatchEnd)
-	{
-		// Gather a range of barriers that all have the same idle flag
-		bool const bIdle = Barriers[BatchEnd].HasIdleFlag();
-
-		while (BatchEnd < Barriers.Num() 
-			&& bIdle == Barriers[BatchEnd].HasIdleFlag())
-		{
-			// Clear the idle flag since its not a valid D3D bit.
-			Barriers[BatchEnd++].ClearIdleFlag();
-		}
-
-		// Insert an idle begin/end timestamp around the barrier batch if required.
-		if (bIdle)
-		{
-			InsertTimestamp(true);
-		}
-
-#if DEBUG_RESOURCE_STATES
-   		TArrayView<FD3D12ResourceBarrier> SubsetBarriers(Barriers.GetData() + BatchStart, BatchEnd - BatchStart);
-   		TConstArrayView<D3D12_RESOURCE_BARRIER> ConstArrayView(	reinterpret_cast<const D3D12_RESOURCE_BARRIER*>(SubsetBarriers.GetData()), SubsetBarriers.Num());
-		LogResourceBarriers(ConstArrayView, CommandList.Interfaces.CommandList.GetReference(), CommandList.QueueType, FString(DX12_RESOURCE_NAME_TO_LOG));
-#endif
-		CommandList.GraphicsCommandList()->ResourceBarrier(BatchEnd - BatchStart, &Barriers[BatchStart]);
-#if DEBUG_RESOURCE_STATES
-		// Keep track of all the resource barriers that have been submitted to the current command list.
-		for(int i = BatchStart; i < BatchEnd - BatchStart - 1; i++)
-		{
-			CommandList.State.ResourceBarriers.Emplace(Barriers[i]);
-		}
-#endif // #if DEBUG_RESOURCE_STATES
-
-		if (bIdle)
-		{
-			InsertTimestamp(false);
-		}
-	}
-
-	Barriers.Reset();
-}
-
 void FD3D12DynamicRHI::RHIReplaceResources(FRHICommandListBase& RHICmdList, TArray<FRHIResourceReplaceInfo>&& ReplaceInfos)
 {
 	RHICmdList.EnqueueLambdaMultiPipe(GetEnabledRHIPipelines(), FRHICommandListBase::EThreadFence::Enabled, TEXT("FD3D12DynamicRHI::RHIReplaceResources"),
diff --git a/Engine/Source/Runtime/D3D12RHI/Private/D3D12Resources.h b/Engine/Source/Runtime/D3D12RHI/Private/D3D12Resources.h
--- a/Engine/Source/Runtime/D3D12RHI/Private/D3D12Resources.h
+++ b/Engine/Source/Runtime/D3D12RHI/Private/D3D12Resources.h
@@ -16,13 +16,6 @@
 #include "RHIPoolAllocator.h"
 #include "Templates/UniquePtr.h"
 
-constexpr D3D12_RESOURCE_STATES BackBufferBarrierWriteTransitionTargets = D3D12_RESOURCE_STATES(
-	uint32(D3D12_RESOURCE_STATE_RENDER_TARGET) |
-	uint32(D3D12_RESOURCE_STATE_UNORDERED_ACCESS) |
-	uint32(D3D12_RESOURCE_STATE_STREAM_OUT) |
-	uint32(D3D12_RESOURCE_STATE_COPY_DEST) |
-	uint32(D3D12_RESOURCE_STATE_RESOLVE_DEST));
-
 // Forward Decls
 class FD3D12CommandList;
 class FD3D12Resource;
@@ -189,12 +182,13 @@ class FD3D12Resource : public FThreadSafeRefCountedObject, public FD3D12DeviceCh
 	friend class FD3D12Buffer;
 private:
 	TRefCountPtr<ID3D12Resource> Resource;
+	TRefCountPtr<FD3D12Heap> Heap;
+
 #if !D3D12RHI_SUPPORTS_UNCOMPRESSED_UAV
 	// Since certain formats cannot be aliased in D3D12, we create a placed UAV
 	// that aliases them. This resource description is used to create that UAV.
 	TUniquePtr<D3D12_RESOURCE_DESC> UAVAccessResourceDesc;
 #endif
-	TRefCountPtr<FD3D12Heap> Heap;
 
 	FD3D12ResidencyHandle* ResidencyHandle = nullptr; // Residency handle owned by this object
 
@@ -204,12 +198,9 @@ private:
 #if NV_AFTERMATH
 	UE::RHICore::Nvidia::Aftermath::D3D12::FResource AftermathHandle{};
 #endif
-
 	const FD3D12ResourceDesc Desc;
-	D3D12_RESOURCE_STATES InitialResourceState{ D3D12_RESOURCE_STATE_TBD };
-	D3D12_RESOURCE_STATES DefaultResourceState{ D3D12_RESOURCE_STATE_TBD };
-	D3D12_RESOURCE_STATES ReadableState{ D3D12_RESOURCE_STATE_CORRUPT };
-	D3D12_RESOURCE_STATES WritableState{ D3D12_RESOURCE_STATE_CORRUPT };
+	ED3D12Access InitialD3D12Access = ED3D12Access::Unknown;
+	ED3D12Access DefaultD3D12Access = ED3D12Access::Unknown;
 #ifdef PLATFORM_SUPPORTS_RESOURCE_COMPRESSION
 	D3D12_RESOURCE_STATES CompressedState{ D3D12_RESOURCE_STATE_CORRUPT };
 #endif
@@ -251,7 +242,7 @@ public:
 	explicit FD3D12Resource(FD3D12Device* ParentDevice,
 		FRHIGPUMask VisibleNodes,
 		ID3D12Resource* InResource,
-		D3D12_RESOURCE_STATES InInitialResourceState,
+		ED3D12Access InInitialD3D12Access,
 		FD3D12ResourceDesc const& InDesc,
 		FD3D12Heap* InHeap = nullptr,
 		D3D12_HEAP_TYPE InHeapType = D3D12_HEAP_TYPE_DEFAULT);
@@ -259,9 +250,9 @@ public:
 	explicit FD3D12Resource(FD3D12Device* ParentDevice,
 		FRHIGPUMask VisibleNodes,
 		ID3D12Resource* InResource,
-		D3D12_RESOURCE_STATES InInitialResourceState,
+		ED3D12Access InInitialD3D12Access,
 		ED3D12ResourceStateMode InResourceStateMode,
-		D3D12_RESOURCE_STATES InDefaultResourceState,
+		ED3D12Access InDefaultD3D12Access,
 		FD3D12ResourceDesc const& InDesc,
 		FD3D12Heap* InHeap,
 		D3D12_HEAP_TYPE InHeapType);
@@ -325,10 +316,8 @@ public:
 	uint16 GetArraySize() const { return (Desc.Dimension == D3D12_RESOURCE_DIMENSION_TEXTURE3D) ? 1 : Desc.DepthOrArraySize; }
 	uint8 GetPlaneCount() const { return PlaneCount; }
 	uint16 GetSubresourceCount() const { return SubresourceCount; }
-	D3D12_RESOURCE_STATES GetInitialResourceState() const { return InitialResourceState; }
-	D3D12_RESOURCE_STATES GetDefaultResourceState() const { check(!bRequiresResourceStateTracking); return DefaultResourceState; }
-	D3D12_RESOURCE_STATES GetWritableState() const { return WritableState; }
-	D3D12_RESOURCE_STATES GetReadableState() const { return ReadableState; }
+	ED3D12Access GetInitialAccess() const { return InitialD3D12Access; }
+	ED3D12Access GetDefaultAccess() const { return DefaultD3D12Access; }
 #ifdef PLATFORM_SUPPORTS_RESOURCE_COMPRESSION
 	D3D12_RESOURCE_STATES GetCompressedState() const { return CompressedState; }
 	void SetCompressedState(D3D12_RESOURCE_STATES State) { CompressedState = State; }
@@ -459,50 +448,41 @@ public:
 			bReadBackResource(HeapType == D3D12_HEAP_TYPE_READBACK)
 		{}
 
-		const D3D12_RESOURCE_STATES GetOptimalInitialState(ERHIAccess InResourceState, bool bAccurateWriteableStates) const
+		ED3D12Access GetOptimalInitialD3D12Access(ED3D12Access InD3D12Access, bool bAccurateWriteableStates) const
 		{
 			// Ignore the requested resource state for non tracked resource because RHI will assume it's always in default resource 
 			// state then when a transition is required (will transition via scoped push/pop to requested state)
-			if (!bSRVOnly && InResourceState != ERHIAccess::Unknown && InResourceState != ERHIAccess::Discard)
+			if (!bSRVOnly && InD3D12Access != ED3D12Access::Unknown)
 			{
-				return GetD3D12ResourceState(InResourceState, ED3D12QueueType::Direct, nullptr);
+				return InD3D12Access;
 			}
-			else
+
+			if (bSRVOnly)
+			{
+				return ED3D12Access::SRVGraphics;
+			}
+			else if (bBuffer && !bUAV)
+			{
+				return (bReadBackResource) ? ED3D12Access::CopyDest : ED3D12Access::GenericRead;
+			}
+			else if (bWritable && bAccurateWriteableStates)
 			{
-				if (bSRVOnly)
+				if (bDSV)
 				{
-					return D3D12_RESOURCE_STATE_NON_PIXEL_SHADER_RESOURCE | D3D12_RESOURCE_STATE_PIXEL_SHADER_RESOURCE;
+					return ED3D12Access::DSVWrite;
 				}
-				else if (bBuffer && !bUAV)
+				else if (bRTV)
 				{
-					return (bReadBackResource) ? D3D12_RESOURCE_STATE_COPY_DEST : D3D12_RESOURCE_STATE_GENERIC_READ;
+					return ED3D12Access::RTV;
 				}
-				else if (bWritable)
+				else if (bUAV)
 				{
-					if (bAccurateWriteableStates)
-					{
-						if (bDSV)
-						{
-							return D3D12_RESOURCE_STATE_DEPTH_WRITE;
-						}
-						else if (bRTV)
-						{
-							return D3D12_RESOURCE_STATE_RENDER_TARGET;
-						}
-						else if (bUAV)
-						{
-							return D3D12_RESOURCE_STATE_UNORDERED_ACCESS;
-						}
-					}
-					else
-					{
-						// This things require tracking anyway
-						return D3D12_RESOURCE_STATE_COMMON;
-					}
+					check(InD3D12Access != ED3D12Access::Unknown);
+					return ED3D12Access::UAVMask;
 				}
-
-				return D3D12_RESOURCE_STATE_COMMON;
 			}
+
+			return ED3D12Access::Common;
 		}
 
 		const uint32 bSRV : 1;
@@ -520,40 +500,53 @@ public:
 	void CommitReservedResource(ID3D12CommandQueue* D3DCommandQueue, uint64 RequiredCommitSizeInBytes);
 
 private:
-	void InitalizeResourceState(D3D12_RESOURCE_STATES InInitialState, ED3D12ResourceStateMode InResourceStateMode, D3D12_RESOURCE_STATES InDefaultState)
+	void InitializeResourceState(
+		const D3D12_HEAP_PROPERTIES* InHeapProps,
+		ED3D12Access InInitialD3D12Access,
+		ED3D12ResourceStateMode InResourceStateMode,
+		ED3D12Access InDefaultD3D12Access)
 	{
+		check(IsValidAccess(InDefaultD3D12Access));
+		check(IsValidAccess(InInitialD3D12Access));
+
 		SubresourceCount = GetMipLevels() * GetArraySize() * GetPlaneCount();
 
 		if (InResourceStateMode == ED3D12ResourceStateMode::SingleState)
 		{
 			// make sure a valid default state is set
-			check(IsValidD3D12ResourceState(InDefaultState));
+			check(InInitialD3D12Access != ED3D12Access::Unknown);
+			check(InDefaultD3D12Access != ED3D12Access::Unknown);
 
 #if UE_BUILD_DEBUG
 			FPlatformAtomics::InterlockedIncrement(&NoStateTrackingResourceCount);
 #endif
-			DefaultResourceState = InDefaultState;
-			WritableState = D3D12_RESOURCE_STATE_CORRUPT;
-			ReadableState = D3D12_RESOURCE_STATE_CORRUPT;
+			InitialD3D12Access = InInitialD3D12Access;
+			DefaultD3D12Access = InDefaultD3D12Access;
 			bRequiresResourceStateTracking = false;
 		}
 		else
 		{
-			DetermineResourceStates(InDefaultState, InResourceStateMode);
+			DetermineResourceStates(InHeapProps, InInitialD3D12Access, InResourceStateMode, InDefaultD3D12Access);
 		}
 
 		if (bRequiresResourceStateTracking)
 		{
 #if D3D12_RHI_RAYTRACING
 			// No state tracking for acceleration structures because they can't have another state
-			check(InDefaultState != D3D12_RESOURCE_STATE_RAYTRACING_ACCELERATION_STRUCTURE && InInitialState != D3D12_RESOURCE_STATE_RAYTRACING_ACCELERATION_STRUCTURE);
+			check(!EnumHasAnyFlags(InInitialD3D12Access, ED3D12Access::BVHRead | ED3D12Access::BVHWrite) ||
+		          !EnumHasAnyFlags(InDefaultD3D12Access, ED3D12Access::BVHRead | ED3D12Access::BVHWrite));
 #endif // D3D12_RHI_RAYTRACING
 
-			InitialResourceState = InInitialState;
+			check(InInitialD3D12Access != ED3D12Access::Unknown);
+			InitialD3D12Access = InInitialD3D12Access;
 		}
 	}
 
-	void DetermineResourceStates(D3D12_RESOURCE_STATES InDefaultState, ED3D12ResourceStateMode InResourceStateMode)
+	void DetermineResourceStates(
+		const D3D12_HEAP_PROPERTIES* InHeapProps,
+		ED3D12Access InInitialD3D12Access,
+		ED3D12ResourceStateMode InResourceStateMode,
+		ED3D12Access InDefaultD3D12Access)
 	{
 		const FD3D12ResourceTypeHelper Type(Desc, HeapType);
 
@@ -563,48 +556,29 @@ private:
 		SetCompressedState(D3D12_RESOURCE_STATE_COMMON);
 #endif
 
-		if (Type.bWritable || InResourceStateMode == ED3D12ResourceStateMode::MultiState)
-		{
-			// Determine the resource's write/read states.
-			if (Type.bRTV)
-			{
-				// Note: The resource could also be used as a UAV however we don't store that writable state. UAV's are handled in a separate RHITransitionResources() specially for UAVs so we know the writeable state in that case should be UAV.
-				check(!Type.bDSV && !Type.bBuffer);
-				WritableState = D3D12_RESOURCE_STATE_RENDER_TARGET;
-				ReadableState = Type.bSRV ? D3D12_RESOURCE_STATE_NON_PIXEL_SHADER_RESOURCE | D3D12_RESOURCE_STATE_PIXEL_SHADER_RESOURCE : D3D12_RESOURCE_STATE_CORRUPT;
-			}
-			else if (Type.bDSV)
-			{
-				check(!Type.bRTV && (!Type.bUAV || GRHISupportsDepthUAV) && !Type.bBuffer);
-				WritableState = D3D12_RESOURCE_STATE_DEPTH_WRITE;
-				ReadableState = Type.bSRV ? D3D12_RESOURCE_STATE_DEPTH_READ | D3D12_RESOURCE_STATE_NON_PIXEL_SHADER_RESOURCE | D3D12_RESOURCE_STATE_PIXEL_SHADER_RESOURCE : D3D12_RESOURCE_STATE_DEPTH_READ;
-			}
-			else
-			{
-				WritableState = Type.bUAV ? D3D12_RESOURCE_STATE_UNORDERED_ACCESS : D3D12_RESOURCE_STATE_CORRUPT;
-				ReadableState = Type.bSRV ? D3D12_RESOURCE_STATE_NON_PIXEL_SHADER_RESOURCE | D3D12_RESOURCE_STATE_PIXEL_SHADER_RESOURCE : D3D12_RESOURCE_STATE_CORRUPT;
-			}
-		}
-		else
+		if (!Type.bWritable && InResourceStateMode != ED3D12ResourceStateMode::MultiState)
 		{
 			bRequiresResourceStateTracking = false;
 						
 #if UE_BUILD_DEBUG
 			FPlatformAtomics::InterlockedIncrement(&NoStateTrackingResourceCount);
 #endif
-			if (InDefaultState != D3D12_RESOURCE_STATE_TBD)
+			if (InDefaultD3D12Access != ED3D12Access::Unknown)
 			{
-				DefaultResourceState = InDefaultState;
+				DefaultD3D12Access = InDefaultD3D12Access;
 			}
 			else if (Type.bBuffer)
 			{
-				DefaultResourceState = (HeapType == D3D12_HEAP_TYPE_READBACK) ? D3D12_RESOURCE_STATE_COPY_DEST : D3D12_RESOURCE_STATE_GENERIC_READ;
+				DefaultD3D12Access = DetermineInitialBufferD3D12Access(HeapType, InHeapProps);
 			}
 			else
 			{
 				check(Type.bSRVOnly);
-				DefaultResourceState = D3D12_RESOURCE_STATE_NON_PIXEL_SHADER_RESOURCE | D3D12_RESOURCE_STATE_PIXEL_SHADER_RESOURCE;
+				DefaultD3D12Access = ED3D12Access::SRVGraphics;
 			}
+
+			check(InInitialD3D12Access != ED3D12Access::Unknown);
+			InitialD3D12Access = InInitialD3D12Access;
 		}
 	}
 };
@@ -734,7 +708,7 @@ public:
 	FD3D12PoolAllocatorPrivateData&    GetPoolAllocatorPrivateData   ()       { return AllocatorData.PoolAllocatorPrivateData;               }
 
 	// Pool allocation specific functions
-	bool OnAllocationMoved(FD3D12ContextArray const& Contexts, FRHIPoolAllocationData* InNewData, D3D12_RESOURCE_STATES& OutCreateState);
+	bool OnAllocationMoved(FD3D12ContextArray const& Contexts, FRHIPoolAllocationData* InNewData, ED3D12Access& OutRHIAccess);
 	void UnlockPoolData();
 
 	bool IsValid() const { return Type != ResourceLocationType::eUndefined; }
@@ -852,12 +826,31 @@ private:
 struct ID3D12ResourceAllocator
 {
 	// Helper function for textures to compute the correct size and alignment
-	void AllocateTexture(uint32 GPUIndex, D3D12_HEAP_TYPE InHeapType, const FD3D12ResourceDesc& InDesc, EPixelFormat InUEFormat, ED3D12ResourceStateMode InResourceStateMode,
-		D3D12_RESOURCE_STATES InCreateState, const D3D12_CLEAR_VALUE* InClearValue, const TCHAR* InName, FD3D12ResourceLocation& ResourceLocation);
+	void AllocateTexture(
+		uint32 GPUIndex,
+		D3D12_HEAP_TYPE InHeapType,
+		const FD3D12ResourceDesc& InDesc,
+		EPixelFormat InUEFormat,
+		ED3D12Access InInitialD3D12Access,
+		ED3D12ResourceStateMode InResourceStateMode,
+		ED3D12Access InDefaultD3D12Access,
+		const D3D12_CLEAR_VALUE* InClearValue,
+		const TCHAR* InName,
+		FD3D12ResourceLocation& ResourceLocation);
 
 	// Actual pure virtual resource allocation function
-	virtual void AllocateResource(uint32 GPUIndex, D3D12_HEAP_TYPE InHeapType, const FD3D12ResourceDesc& InDesc, uint64 InSize, uint32 InAllocationAlignment, ED3D12ResourceStateMode InResourceStateMode,
-		D3D12_RESOURCE_STATES InCreateState, const D3D12_CLEAR_VALUE* InClearValue, const TCHAR* InName, FD3D12ResourceLocation& ResourceLocation) = 0;
+	virtual void AllocateResource(
+		uint32 GPUIndex,
+		D3D12_HEAP_TYPE InHeapType,
+		const FD3D12ResourceDesc& InDesc,
+		uint64 InSize,
+		uint32 InAllocationAlignment,
+		ED3D12Access InInitialD3D12Access,
+		ED3D12ResourceStateMode InResourceStateMode,
+		ED3D12Access InDefaultD3D12Access,
+		const D3D12_CLEAR_VALUE* InClearValue,
+		const TCHAR* InName,
+		FD3D12ResourceLocation& ResourceLocation) = 0;
 };
 
 struct FD3D12LockedResource : public FD3D12DeviceChild
@@ -1032,7 +1025,12 @@ public:
 
 	virtual ~FD3D12Buffer();
 
-	static void UploadResourceData(FD3D12CommandContext& CommandContext, D3D12_RESOURCE_STATES InDestinationState, FD3D12ResourceLocation& DestinationResourceLocation, const FD3D12ResourceLocation& SourceResourceLocation, uint32 Size);
+	static void UploadResourceData(
+		FD3D12CommandContext& CommandContext,
+		ED3D12Access DestinationD3D12Access,
+		FD3D12ResourceLocation& DestinationResourceLocation,
+		const FD3D12ResourceLocation& SourceResourceLocation,
+		uint32 Size);
 
 	FD3D12SyncPointRef UploadResourceDataViaCopyQueue(FD3D12CommandContext& OwningContext, FResourceArrayUploadInterface* InResourceArray);
 
@@ -1094,40 +1092,6 @@ namespace D3D12BufferStats
 	void UpdateBufferStats(FD3D12Buffer& Buffer, bool bAllocating);
 }
 
-class FD3D12ResourceBarrierBatcher
-{
-	PRAGMA_DISABLE_ENUM_CONSTEXPR_CONVERSION_WARNINGS
-	// Use the top bit of the flags enum to mark transitions as "idle" time (used to remove the swapchain wait time for back buffers).
-	static const D3D12_RESOURCE_BARRIER_FLAGS BarrierFlag_CountAsIdleTime = D3D12_RESOURCE_BARRIER_FLAGS(1ull << ((sizeof(D3D12_RESOURCE_BARRIER_FLAGS) * 8) - 1));
-	PRAGMA_ENABLE_ENUM_CONSTEXPR_CONVERSION_WARNINGS
-
-	struct FD3D12ResourceBarrier : public D3D12_RESOURCE_BARRIER
-	{
-		FD3D12ResourceBarrier() = default;
-		FD3D12ResourceBarrier(D3D12_RESOURCE_BARRIER&& Barrier) : D3D12_RESOURCE_BARRIER(MoveTemp(Barrier)) {}
-
-		bool HasIdleFlag() const { return !!(Flags & BarrierFlag_CountAsIdleTime); }
-		void ClearIdleFlag() { Flags &= ~BarrierFlag_CountAsIdleTime; }
-	};
-	static_assert(sizeof(FD3D12ResourceBarrier) == sizeof(D3D12_RESOURCE_BARRIER), "FD3D12ResourceBarrier is a wrapper to add helper functions. Do not add members.");
-
-public:
-	// Add a UAV barrier to the batch. Ignoring the actual resource for now.
-	void AddUAV();
-
-	// Add a transition resource barrier to the batch. Returns the number of barriers added, which may be negative if an existing barrier was cancelled.
-	int32 AddTransition(FD3D12Resource* pResource, D3D12_RESOURCE_STATES Before, D3D12_RESOURCE_STATES After, uint32 Subresource);
-
-	void AddAliasingBarrier(ID3D12Resource* InResourceBefore, ID3D12Resource* InResourceAfter);
-
-	void FlushIntoCommandList(FD3D12CommandList& CommandList, class FD3D12QueryAllocator& TimestampAllocator);
-
-	int32 Num() const { return Barriers.Num(); }
-
-private:
-	TArray<FD3D12ResourceBarrier> Barriers;
-};
-
 class FD3D12StagingBuffer final : public FRHIStagingBuffer
 {
 	friend class FD3D12CommandContext;
diff --git a/Engine/Source/Runtime/D3D12RHI/Private/D3D12StateCache.cpp b/Engine/Source/Runtime/D3D12RHI/Private/D3D12StateCache.cpp
--- a/Engine/Source/Runtime/D3D12RHI/Private/D3D12StateCache.cpp
+++ b/Engine/Source/Runtime/D3D12RHI/Private/D3D12StateCache.cpp
@@ -188,7 +188,7 @@ void FD3D12StateCache::FlushComputeShaderCache(bool bForce)
 {
 	if (bForce)
 	{
-		CmdContext.AddUAVBarrier();
+		CmdContext.AddGlobalBarrier(ED3D12Access::UAVCompute, ED3D12Access::UAVCompute);
 		INC_DWORD_STAT(STAT_D3D12UAVBarriers);
 	}
 }
