D3D12 Enhanced Barriers

diff --git a/Engine/Source/Runtime/D3D12RHI/Private/D3D12Texture.cpp b/Engine/Source/Runtime/D3D12RHI/Private/D3D12Texture.cpp
--- a/Engine/Source/Runtime/D3D12RHI/Private/D3D12Texture.cpp
+++ b/Engine/Source/Runtime/D3D12RHI/Private/D3D12Texture.cpp
@@ -216,6 +216,11 @@ bool FD3D12Texture::CanBe4KAligned(const FD3D12ResourceDesc& Desc, EPixelFormat
 
 using namespace D3D12RHI;
 
+void FD3D12DynamicRHI::SetFormatAliasedTexturesMustBeCreatedUsingCommonLayout(bool bValue)
+{
+	bFormatAliasedTexturesMustBeCreatedUsingCommonLayout = bValue;
+}
+
 FD3D12ResourceDesc FD3D12DynamicRHI::GetResourceDesc(const FRHITextureDesc& TextureDesc) const
 {
 	FD3D12ResourceDesc ResourceDesc;
@@ -473,7 +478,8 @@ bool FD3D12DynamicRHI::RHIGetTextureMemoryVisualizeData(FColor* /*TextureData*/,
 /**
  * Creates a 2D texture optionally guarded by a structured exception handler.
  */
-void SafeCreateTexture2D(FD3D12Device* pDevice, 
+void SafeCreateTexture2D(
+	FD3D12Device* pDevice, 
 	FD3D12Adapter* Adapter,
 	const FD3D12ResourceDesc& TextureDesc,
 	const D3D12_CLEAR_VALUE* ClearValue, 
@@ -481,7 +487,8 @@ void SafeCreateTexture2D(FD3D12Device* pDevice,
 	FD3D12BaseShaderResource* Owner,
 	EPixelFormat Format,
 	ETextureCreateFlags Flags,
-	D3D12_RESOURCE_STATES InitialState,
+	ED3D12Access InInitialD3D12Access,
+	ED3D12Access InDefaultD3D12Access,
 	const TCHAR* Name)
 {
 
@@ -512,8 +519,18 @@ void SafeCreateTexture2D(FD3D12Device* pDevice,
 			{
 				FD3D12Resource* Resource = nullptr;
 				VERIFYD3D12CREATETEXTURERESULT(
-					Adapter->CreateReservedResource(TextureDesc, pDevice->GetGPUMask(), InitialState, ED3D12ResourceStateMode::MultiState, InitialState, ClearValue, &Resource, Name, false),
-					TextureDesc, pDevice->GetDevice());
+					Adapter->CreateReservedResource(
+						TextureDesc,
+						pDevice->GetGPUMask(),
+						InInitialD3D12Access,
+						ED3D12ResourceStateMode::MultiState,
+						InDefaultD3D12Access,
+						ClearValue,
+						&Resource,
+						Name,
+						false),
+					TextureDesc,
+					pDevice->GetDevice());
 
 				D3D12_RESOURCE_ALLOCATION_INFO AllocInfo = pDevice->GetResourceAllocationInfo(TextureDesc);
 
@@ -528,7 +545,17 @@ void SafeCreateTexture2D(FD3D12Device* pDevice,
 			}
 			else
 			{
-				VERIFYD3D12CREATETEXTURERESULT(pDevice->GetTextureAllocator().AllocateTexture(TextureDesc, ClearValue, Format, *OutTexture2D, InitialState, Name), TextureDesc, pDevice->GetDevice());
+				VERIFYD3D12CREATETEXTURERESULT(
+					pDevice->GetTextureAllocator().AllocateTexture(
+						TextureDesc,
+						ClearValue,
+						Format,
+						*OutTexture2D,
+						InInitialD3D12Access,
+						InDefaultD3D12Access,
+						Name),
+					TextureDesc,
+					pDevice->GetDevice());
 			}
 
 			OutTexture2D->SetOwner(Owner);
@@ -615,15 +642,25 @@ FD3D12DynamicRHI::FCreateTextureInternalResult FD3D12DynamicRHI::CreateTextureIn
 	D3D12_CLEAR_VALUE* ClearValuePtr = UE::D3D12RHI::TextureUtils::FillClearValue(ClearValue, ResourceDesc, CreateDesc);
 
 	const FD3D12Resource::FD3D12ResourceTypeHelper Type(ResourceDesc, D3D12_HEAP_TYPE_DEFAULT);
-	const D3D12_RESOURCE_STATES DesiredState = Type.GetOptimalInitialState(CreateDesc.InitialState, true);
+	const ED3D12Access DesiredD3D12Access = Type.GetOptimalInitialD3D12Access(ConvertToD3D12Access(CreateDesc.InitialState), true);
 
-	const D3D12_RESOURCE_STATES CreateState =
-		// When creating a resource with castable formats, d3d12 uses enhanced barriers behind the scenes which means we have to start it in the "common" state to use "legacy" barriers
-		ResourceDesc.SupportsUncompressedUAV() ? D3D12_RESOURCE_STATE_COMMON
-		// If we have initial data, we want the resource created in COPY_DEST so we can copy the data immediately
-		: bCreateAsCopyDest ? D3D12_RESOURCE_STATE_COPY_DEST
-		// Otherwise, create the resource in the caller's specified state
-		: DesiredState;
+	const ED3D12Access InitialD3D12Access = [DesiredD3D12Access, bCreateAsCopyDest, &ResourceDesc]() {
+		if (bFormatAliasedTexturesMustBeCreatedUsingCommonLayout
+			&& ResourceDesc.SupportsUncompressedUAV())
+		{
+			// When creating a resource with castable formats, d3d12 uses enhanced barriers behind the scenes
+			// which means we have to start it in the "common" state to then transition it to "legacy" barriers
+			return ED3D12Access::Common;
+		}
+		else
+		{
+			return bCreateAsCopyDest ?
+				// If we have initial data, we want the resource created in COPY_DEST so we can copy the data immediately
+				ED3D12Access::CopyDest :
+				// Otherwise, create the resource in the calculated optimal state
+				DesiredD3D12Access;
+		}
+	}();
 
 	FD3D12Adapter* Adapter = &GetAdapter();
 	FD3D12Texture* D3D12TextureOut = Adapter->CreateLinkedObject<FD3D12Texture>(CreateDesc.GPUMask, [&](FD3D12Device* Device, FD3D12Texture* FirstLinkedObject)
@@ -643,7 +680,17 @@ FD3D12DynamicRHI::FCreateTextureInternalResult FD3D12DynamicRHI::CreateTextureIn
 		if (ResourceAllocator)
 		{
 			const D3D12_HEAP_TYPE HeapType = EnumHasAnyFlags(CreateDesc.Flags, TexCreate_CPUReadback) ? D3D12_HEAP_TYPE_READBACK : D3D12_HEAP_TYPE_DEFAULT;
-			ResourceAllocator->AllocateTexture(Device->GetGPUIndex(), HeapType, ResourceDesc, (EPixelFormat)CreateDesc.Format, ED3D12ResourceStateMode::Default, CreateState, ClearValuePtr, CreateDesc.DebugName, Location);
+			ResourceAllocator->AllocateTexture(
+				Device->GetGPUIndex(),
+				HeapType,
+				ResourceDesc,
+				(EPixelFormat)CreateDesc.Format,
+				InitialD3D12Access,
+				ED3D12ResourceStateMode::Default,
+				DesiredD3D12Access,
+				ClearValuePtr,
+				CreateDesc.DebugName,
+				Location);
 			Location.SetOwner(NewTexture);
 		}
 		else if(EnumHasAnyFlags(CreateDesc.Flags, TexCreate_CPUReadback))
@@ -665,7 +712,19 @@ FD3D12DynamicRHI::FCreateTextureInternalResult FD3D12DynamicRHI::CreateTextureIn
 			if (ResourceDesc.bReservedResource)
 			{
 				FD3D12Resource* Resource = nullptr;
-				VERIFYD3D12CREATETEXTURERESULT(Adapter->CreateReservedResource(ResourceDesc, Device->GetGPUMask(), CreateState, ED3D12ResourceStateMode::MultiState, CreateState, ClearValuePtr, &Resource, CreateDesc.DebugName, false), ResourceDesc, Device->GetDevice());
+				VERIFYD3D12CREATETEXTURERESULT(
+					Adapter->CreateReservedResource(
+						ResourceDesc,
+						Device->GetGPUMask(),
+						InitialD3D12Access,
+						ED3D12ResourceStateMode::MultiState,
+						DesiredD3D12Access,
+						ClearValuePtr,
+						&Resource,
+						CreateDesc.DebugName,
+						false),
+					ResourceDesc,
+					Device->GetDevice());
 
 				D3D12_RESOURCE_ALLOCATION_INFO AllocInfo = Device->GetResourceAllocationInfo(ResourceDesc);
 
@@ -678,14 +737,25 @@ FD3D12DynamicRHI::FCreateTextureInternalResult FD3D12DynamicRHI::CreateTextureIn
 			}
 			else
 			{
-				VERIFYD3D12CREATETEXTURERESULT(Device->GetTextureAllocator().AllocateTexture(ResourceDesc, ClearValuePtr, CreateDesc.Format, Location, CreateState, CreateDesc.DebugName), ResourceDesc, Device->GetDevice());
+				VERIFYD3D12CREATETEXTURERESULT(
+					Device->GetTextureAllocator().AllocateTexture(
+						ResourceDesc,
+						ClearValuePtr,
+						CreateDesc.Format,
+						Location,
+						InitialD3D12Access,
+						DesiredD3D12Access,
+						CreateDesc.DebugName),
+					ResourceDesc,
+					Device->GetDevice());
 			}
 
 			Location.SetOwner(NewTexture);
 		}
 		else
 		{
-			SafeCreateTexture2D(Device,
+			SafeCreateTexture2D(
+				Device,
 				Adapter,
 				ResourceDesc,
 				ClearValuePtr,
@@ -693,7 +763,8 @@ FD3D12DynamicRHI::FCreateTextureInternalResult FD3D12DynamicRHI::CreateTextureIn
 				NewTexture,
 				CreateDesc.Format,
 				CreateDesc.Flags,
-				CreateState,
+				InitialD3D12Access,
+				DesiredD3D12Access,
 				CreateDesc.DebugName);
 		}
 
@@ -728,8 +799,8 @@ FD3D12DynamicRHI::FCreateTextureInternalResult FD3D12DynamicRHI::CreateTextureIn
 
 	FCreateTextureInternalResult Result{};
 	Result.Texture = D3D12TextureOut;
-	Result.CreateState = CreateState;
-	Result.DesiredState = DesiredState;
+	Result.CreateD3D12Access = InitialD3D12Access;
+	Result.DesiredD3D12Access = DesiredD3D12Access;
 
 	return Result;
 #else
@@ -739,18 +810,26 @@ FD3D12DynamicRHI::FCreateTextureInternalResult FD3D12DynamicRHI::CreateTextureIn
 #endif // PLATFORM_WINDOWS
 }
 
-void UE::D3D12RHI::TextureUtils::ReconcileInitialState(FRHICommandListBase& RHICmdList, FD3D12Texture* Texture, D3D12_RESOURCE_STATES CurrentState, D3D12_RESOURCE_STATES NeededState)
+void UE::D3D12RHI::TextureUtils::ReconcileInitialState(
+	FRHICommandListBase& RHICmdList,
+	FD3D12Texture* Texture,
+	ED3D12Access CurrentD3D12Access,
+	ED3D12Access NeededD3D12Access)
 {
-	if (CurrentState != NeededState)
+	if (CurrentD3D12Access != NeededD3D12Access)
 	{
-		RHICmdList.EnqueueLambda([Texture, CurrentState, NeededState](FRHICommandListBase& ExecutingCmdList)
+		RHICmdList.EnqueueLambda([Texture, CurrentD3D12Access, NeededD3D12Access](FRHICommandListBase& ExecutingCmdList)
 		{
 			for (FD3D12Texture::FLinkedObjectIterator TextureIt(Texture); TextureIt; ++TextureIt)
 			{
 				const uint32 GPUIndex = TextureIt->GetParentDevice()->GetGPUIndex();
 
 				FD3D12CommandContext& Context = FD3D12CommandContext::Get(ExecutingCmdList, GPUIndex);
-				Context.AddTransitionBarrier(TextureIt->GetResource(), CurrentState, NeededState, D3D12_RESOURCE_BARRIER_ALL_SUBRESOURCES);
+				Context.AddBarrier(
+					TextureIt->GetResource(),
+					CurrentD3D12Access,
+					NeededD3D12Access,
+					D3D12_RESOURCE_BARRIER_ALL_SUBRESOURCES);
 			}
 		});
 	}
@@ -895,11 +974,11 @@ static UE::RHICore::FBaseTextureInitializerImplementation CreateTextureInitializ
 	void* WritableData = UploadLocation.GetMappedBaseAddress();
 
 	return UE::RHICore::FBaseTextureInitializerImplementation(RHICmdList, Texture, WritableData, TextureSize,
-		[Texture = TRefCountPtr<FD3D12Texture>(Texture), CreateState = CreateResult.CreateState, DesiredState = CreateResult.DesiredState, UploadLocation = MoveTemp(UploadLocation)](FRHICommandListBase& RHICmdList) mutable
+		[Texture = TRefCountPtr<FD3D12Texture>(Texture), CreateD3D12Access = CreateResult.CreateD3D12Access, DesiredD3D12Access = CreateResult.DesiredD3D12Access, UploadLocation = MoveTemp(UploadLocation)](FRHICommandListBase& RHICmdList) mutable
 		{
-			UE::D3D12RHI::TextureUtils::ReconcileInitialState(RHICmdList, Texture, CreateState, D3D12_RESOURCE_STATE_COPY_DEST);
+			UE::D3D12RHI::TextureUtils::ReconcileInitialState(RHICmdList, Texture, CreateD3D12Access, ED3D12Access::CopyDest);
 
-			Texture->UploadInitialData(RHICmdList, MoveTemp(UploadLocation), DesiredState);
+			Texture->UploadInitialData(RHICmdList, MoveTemp(UploadLocation), DesiredD3D12Access);
 			return TRefCountPtr<FRHITexture>(MoveTemp(Texture));
 		},
 		[Texture = TRefCountPtr<FD3D12Texture>(Texture), WritableData](FRHITextureInitializer::FSubresourceIndex SubresourceIndex)
@@ -950,9 +1029,9 @@ FRHITextureInitializer FD3D12DynamicRHI::RHICreateTextureInitializer(FRHICommand
 	{
 		// Just return the texture with its default contents
 		return UE::RHICore::FBaseTextureInitializerImplementation(RHICmdList, CreateResult.Texture,
-			[Texture = TRefCountPtr<FD3D12Texture>(CreateResult.Texture), CreateState = CreateResult.CreateState, DesiredState = CreateResult.DesiredState](FRHICommandListBase& RHICmdList) mutable
+			[Texture = TRefCountPtr<FD3D12Texture>(CreateResult.Texture), CreateD3D12Access = CreateResult.CreateD3D12Access, DesiredD3D12Access = CreateResult.DesiredD3D12Access](FRHICommandListBase& RHICmdList) mutable
 			{
-				UE::D3D12RHI::TextureUtils::ReconcileInitialState(RHICmdList, Texture, CreateState, DesiredState);
+				UE::D3D12RHI::TextureUtils::ReconcileInitialState(RHICmdList, Texture, CreateD3D12Access, DesiredD3D12Access);
 				return TRefCountPtr<FRHITexture>(MoveTemp(Texture));
 			}
 		);
@@ -1082,11 +1161,6 @@ FTextureRHIRef FD3D12DynamicRHI::RHIAsyncCreateTexture2D(uint32 SizeX, uint32 Si
 		SubResourceData[MipIndex].SlicePitch = MipSize;
 	}
 
-	// All resources used in a COPY command list must begin in the COMMON state. 
-	// COPY_SOURCE and COPY_DEST are "promotable" states. You can create async texture resources in the COMMON state and still avoid any state transitions by relying on state promotion. 
-	// Also remember that ALL touched resources in a COPY command list decay to COMMON after ExecuteCommandLists completes.
-	const D3D12_RESOURCE_STATES InitialState = D3D12_RESOURCE_STATE_COMMON;
-
 #if !PLATFORM_WINDOWS
 	FRWScopeLock ReadLock(*RHIGetSuspendedLock(), SLT_ReadOnly);
 #endif
@@ -1104,7 +1178,9 @@ FTextureRHIRef FD3D12DynamicRHI::RHIAsyncCreateTexture2D(uint32 SizeX, uint32 Si
 			NewTexture,
 			(EPixelFormat)Format,
 			Flags,
-			InitialState,
+			// All resources used in a COPY command list must begin in the COMMON state. 
+			ED3D12Access::Common,
+			ED3D12Access::SRVMask,
 			nullptr);
 
 		D3D12_SHADER_RESOURCE_VIEW_DESC SRVDesc = {};
@@ -1478,7 +1554,7 @@ FTextureRHIRef FD3D12DynamicRHI::AsyncReallocateTexture2D_RenderThread(FRHIComma
 
 	// Allocate a new texture.
 	const FCreateTextureInternalResult CreateResult = CreateTextureInternal(CreateDesc, nullptr);
-	UE::D3D12RHI::TextureUtils::ReconcileInitialState(RHICmdList, CreateResult.Texture, CreateResult.CreateState, CreateResult.DesiredState);
+	UE::D3D12RHI::TextureUtils::ReconcileInitialState(RHICmdList, CreateResult.Texture, CreateResult.CreateD3D12Access, CreateResult.DesiredD3D12Access);
 
 	FD3D12Texture* NewTexture = CreateResult.Texture;
 
@@ -1501,8 +1577,16 @@ FTextureRHIRef FD3D12DynamicRHI::AsyncReallocateTexture2D_RenderThread(FRHIComma
 
 			FD3D12CommandContext& Context = FD3D12CommandContext::Get(ExecutingCmdList, DeviceOldTexture->GetParentDevice()->GetGPUIndex());
 
-			FScopedResourceBarrier ScopeResourceBarrierDst(Context, DeviceNewTexture->GetResource(), &DeviceNewTexture->ResourceLocation, D3D12_RESOURCE_STATE_COPY_DEST, D3D12_RESOURCE_BARRIER_ALL_SUBRESOURCES);
-			FScopedResourceBarrier ScopeResourceBarrierSrc(Context, DeviceOldTexture->GetResource(), &DeviceOldTexture->ResourceLocation, D3D12_RESOURCE_STATE_COPY_SOURCE, D3D12_RESOURCE_BARRIER_ALL_SUBRESOURCES);
+			FScopedResourceBarrier ScopeResourceBarrierDst(
+				Context,
+				DeviceNewTexture->GetResource(),
+				ED3D12Access::CopyDest,
+				D3D12_RESOURCE_BARRIER_ALL_SUBRESOURCES);
+			FScopedResourceBarrier ScopeResourceBarrierSrc(
+				Context,
+				DeviceOldTexture->GetResource(),
+				ED3D12Access::CopySrc,
+				D3D12_RESOURCE_BARRIER_ALL_SUBRESOURCES);
 			Context.FlushResourceBarriers();	// Must flush so the desired state is actually set.
 
 			for (uint32 MipIndex = 0; MipIndex < NumSharedMips; ++MipIndex)
@@ -2049,7 +2133,11 @@ void FD3D12Texture::UpdateTexture(FD3D12CommandContext& Context, uint32 MipIndex
 {
 	LLM_SCOPE_BYNAME(TEXT("D3D12CopyTextureRegion"));
 
-	FScopedResourceBarrier ScopeResourceBarrierDest(Context, GetResource(), &ResourceLocation, D3D12_RESOURCE_STATE_COPY_DEST, MipIndex);
+	FScopedResourceBarrier ScopeResourceBarrierDest(
+		Context,
+		GetResource(),
+		ED3D12Access::CopyDest,
+		MipIndex);
 	// Don't need to transition upload heaps
 
 	CD3DX12_TEXTURE_COPY_LOCATION DestCopyLocation(GetResource()->GetResource(), MipIndex);
@@ -2077,8 +2165,17 @@ void FD3D12Texture::CopyTextureRegion(FD3D12CommandContext& Context, uint32 Dest
 	CD3DX12_TEXTURE_COPY_LOCATION DestCopyLocation(GetResource()->GetResource(), 0);
 	CD3DX12_TEXTURE_COPY_LOCATION SourceCopyLocation(SourceTexture->GetResource()->GetResource(), 0);
 
-	FScopedResourceBarrier ConditionalScopeResourceBarrierDst(Context, GetResource(), &ResourceLocation, D3D12_RESOURCE_STATE_COPY_DEST, DestCopyLocation.SubresourceIndex);
-	FScopedResourceBarrier ConditionalScopeResourceBarrierSrc(Context, SourceTexture->GetResource(), &SourceTexture->ResourceLocation, D3D12_RESOURCE_STATE_COPY_SOURCE, SourceCopyLocation.SubresourceIndex);
+	FScopedResourceBarrier ConditionalScopeResourceBarrierDst(
+		Context,
+		GetResource(),
+		ED3D12Access::CopyDest,
+		DestCopyLocation.SubresourceIndex);
+
+	FScopedResourceBarrier ConditionalScopeResourceBarrierSrc(
+		Context,
+		SourceTexture->GetResource(),
+		ED3D12Access::CopySrc,
+		SourceCopyLocation.SubresourceIndex);
 
 	Context.FlushResourceBarriers();
 	Context.CopyTextureRegionChecked(
@@ -2095,9 +2192,12 @@ void FD3D12Texture::CopyTextureRegion(FD3D12CommandContext& Context, uint32 Dest
 	Context.UpdateResidency(GetResource());
 }
 
-void FD3D12Texture::UploadInitialData(FRHICommandListBase& RHICmdList, FD3D12ResourceLocation&& SourceLocation, D3D12_RESOURCE_STATES DestinationState)
+void FD3D12Texture::UploadInitialData(
+	FRHICommandListBase& RHICmdList,
+	FD3D12ResourceLocation&& SourceLocation,
+	ED3D12Access InDestinationD3D12Access)
 {
-	RHICmdList.EnqueueLambda([Texture = this, SourceLocation = MoveTemp(SourceLocation), DestinationState](FRHICommandListBase& ExecutingCmdList)
+	RHICmdList.EnqueueLambda([Texture = this, SourceLocation = MoveTemp(SourceLocation), InDestinationD3D12Access](FRHICommandListBase& ExecutingCmdList)
 	{
 		const uint32 NumSubresources = CalculateSubresourceCount(Texture->GetDesc());
 
@@ -2146,14 +2246,18 @@ void FD3D12Texture::UploadInitialData(FRHICommandListBase& RHICmdList, FD3D12Res
 			}
 
 			// Update the resource state after the copy has been done (will take care of updating the residency as well)
-			Context.AddTransitionBarrier(Resource, D3D12_RESOURCE_STATE_COPY_DEST, DestinationState, D3D12_RESOURCE_BARRIER_ALL_SUBRESOURCES);
+			Context.AddBarrier(
+				Resource,
+				ED3D12Access::CopyDest,
+				InDestinationD3D12Access,
+				D3D12_RESOURCE_BARRIER_ALL_SUBRESOURCES);
 			Context.ConditionalSplitCommandList();
 
 			// Texture is now written and ready, so unlock the block (locked after creation and can be defragmented if needed)
 			CurrentTexture.ResourceLocation.UnlockPoolData();
 
 			// If the resource is untracked, the destination state must match the default state of the resource.
-			check(Resource->RequiresResourceStateTracking() || (Resource->GetDefaultResourceState() == DestinationState));
+			check(Resource->RequiresResourceStateTracking() || (Resource->GetDefaultAccess() == InDestinationD3D12Access));
 		}
 	});
 }
@@ -2461,7 +2565,11 @@ public:
 
 			CD3DX12_TEXTURE_COPY_LOCATION DestCopyLocation(TextureLink.GetResource()->GetResource(), MipIdx);
 
-			FScopedResourceBarrier ScopeResourceBarrierDest(Context, TextureLink.GetResource(), &TextureLink.ResourceLocation, D3D12_RESOURCE_STATE_COPY_DEST, DestCopyLocation.SubresourceIndex);
+			FScopedResourceBarrier ScopeResourceBarrierDest(
+				Context,
+				TextureLink.GetResource(),
+				ED3D12Access::CopyDest,
+				DestCopyLocation.SubresourceIndex);
 
 			Context.FlushResourceBarriers();
 
@@ -2703,7 +2811,11 @@ void FD3D12DynamicRHI::EndUpdateTexture3D_Internal(FRHICommandListBase& RHICmdLi
 				CD3DX12_TEXTURE_COPY_LOCATION DestCopyLocation(Texture.GetResource()->GetResource(), MipIdx);
 				CD3DX12_TEXTURE_COPY_LOCATION SourceCopyLocation(SrcResourceLocation->GetResource()->GetResource(), PlacedSubresourceFootprint);
 
-				FScopedResourceBarrier ScopeResourceBarrierDest(Context, Texture.GetResource(), &Texture.ResourceLocation, D3D12_RESOURCE_STATE_COPY_DEST, DestCopyLocation.SubresourceIndex);
+				FScopedResourceBarrier ScopeResourceBarrierDest(
+					Context,
+					Texture.GetResource(),
+					ED3D12Access::CopyDest,
+					DestCopyLocation.SubresourceIndex);
 
 				Context.FlushResourceBarriers();
 				Context.UpdateResidency(Texture.GetResource());
@@ -2807,10 +2919,10 @@ FD3D12Texture* FD3D12DynamicRHI::CreateTextureFromResource(bool bTextureArray, b
 
 	// The state this resource will be in when it leaves this function
 	const FD3D12Resource::FD3D12ResourceTypeHelper Type(TextureDesc, D3D12_HEAP_TYPE_DEFAULT);
-	const D3D12_RESOURCE_STATES DestinationState = Type.GetOptimalInitialState(ERHIAccess::Unknown, !EnumHasAnyFlags(TexCreateFlags, TexCreate_Shared));
+	const ED3D12Access DestinationD3D12Access = Type.GetOptimalInitialD3D12Access(ED3D12Access::Unknown, !EnumHasAnyFlags(TexCreateFlags, TexCreate_Shared));
 
 	FD3D12Device* Device = Adapter->GetDevice(0);
-	FD3D12Resource* TextureResource = new FD3D12Resource(Device, Device->GetGPUMask(), Resource, DestinationState, TextureDesc);
+	FD3D12Resource* TextureResource = new FD3D12Resource(Device, Device->GetGPUMask(), Resource, DestinationD3D12Access, TextureDesc);
 	TextureResource->AddRef();
 	TextureResource->SetName(TEXT("TextureFromResource"));
 
@@ -2947,8 +3059,16 @@ void FD3D12CommandContext::RHICopyTexture(FRHITexture* SourceTextureRHI, FRHITex
 	FD3D12Texture* SourceTexture = RetrieveTexture(SourceTextureRHI);
 	FD3D12Texture* DestTexture = RetrieveTexture(DestTextureRHI);
 
-	FScopedResourceBarrier ConditionalScopeResourceBarrierSrc(*this, SourceTexture->GetResource(), &SourceTexture->ResourceLocation, D3D12_RESOURCE_STATE_COPY_SOURCE, D3D12_RESOURCE_BARRIER_ALL_SUBRESOURCES);
-	FScopedResourceBarrier ConditionalScopeResourceBarrierDst(*this, DestTexture->GetResource(), &DestTexture->ResourceLocation, D3D12_RESOURCE_STATE_COPY_DEST  , D3D12_RESOURCE_BARRIER_ALL_SUBRESOURCES);
+	FScopedResourceBarrier ConditionalScopeResourceBarrierSrc(
+		*this,
+		SourceTexture->GetResource(),
+		ED3D12Access::CopySrc,
+		D3D12_RESOURCE_BARRIER_ALL_SUBRESOURCES);
+	FScopedResourceBarrier ConditionalScopeResourceBarrierDst(
+		*this,
+		DestTexture->GetResource(),
+		ED3D12Access::CopyDest,
+		D3D12_RESOURCE_BARRIER_ALL_SUBRESOURCES);
 
 	FlushResourceBarriers();
 
diff --git a/Engine/Source/Runtime/D3D12RHI/Private/D3D12Texture.h b/Engine/Source/Runtime/D3D12RHI/Private/D3D12Texture.h
--- a/Engine/Source/Runtime/D3D12RHI/Private/D3D12Texture.h
+++ b/Engine/Source/Runtime/D3D12RHI/Private/D3D12Texture.h
@@ -13,15 +13,17 @@ D3D12Texture.h: Implementation of D3D12 Texture
 #define GUARDED_TEXTURE_CREATES (PLATFORM_WINDOWS && !(UE_BUILD_SHIPPING || UE_BUILD_TEST || PLATFORM_COMPILER_CLANG))
 
 // @todo don't make global here!
-void SafeCreateTexture2D(FD3D12Device* pDevice, 
+void SafeCreateTexture2D(
+	FD3D12Device* pDevice,
 	FD3D12Adapter* Adapter,
 	const FD3D12ResourceDesc& TextureDesc,
-	const D3D12_CLEAR_VALUE* ClearValue, 
+	const D3D12_CLEAR_VALUE* ClearValue,
 	FD3D12ResourceLocation* OutTexture2D,
 	FD3D12BaseShaderResource* Owner,
 	EPixelFormat Format,
 	ETextureCreateFlags Flags,
-	D3D12_RESOURCE_STATES InitialState,
+	ED3D12Access InInitialD3D12Access,
+	ED3D12Access InDefaultD3D12Access,
 	const TCHAR* Name);
 
 /** D3D12 RHI Texture class */
@@ -86,7 +88,7 @@ public:
 #endif // #if PLATFORM_REQUIRES_TYPELESS_RESOURCE_DISCARD_WORKAROUND
 		
 	// Setup functionality
-	void UploadInitialData(FRHICommandListBase& RHICmdList, FD3D12ResourceLocation&& SourceLocation, D3D12_RESOURCE_STATES DestinationState);
+	void UploadInitialData(FRHICommandListBase& RHICmdList, FD3D12ResourceLocation&& SourceLocation, ED3D12Access InDestinationD3D12Access);
 	void CreateViews(FD3D12Texture* FirstLinkedObject);
 	void SetCreatedRTVsPerSlice(bool Value, int32 InRTVArraySize)
 	{
@@ -277,7 +279,7 @@ namespace UE::D3D12RHI::TextureUtils
 	// Fills in the ClearValue argument based on the ResourceDesc and CreateDesc flags.
 	D3D12_CLEAR_VALUE* FillClearValue(D3D12_CLEAR_VALUE& ClearValue, const FD3D12ResourceDesc& ResourceDesc, const FRHITextureDesc& TextureDesc);
 
-	void ReconcileInitialState(FRHICommandListBase& RHICmdList, FD3D12Texture* Texture, D3D12_RESOURCE_STATES CurrentState, D3D12_RESOURCE_STATES NeededState);
+	void ReconcileInitialState(FRHICommandListBase& RHICmdList, FD3D12Texture* Texture, ED3D12Access CurrentD3D12Access, ED3D12Access NeededD3D12Access);
 
 	uint64 CalculateResourceSize(FD3D12Texture* Texture);
 	void CopyBulkData(void* UploadMemory, uint64 UploadMemorySize, FD3D12Texture* Texture, FResourceBulkDataInterface* BulkData);
diff --git a/Engine/Source/Runtime/D3D12RHI/Private/D3D12TransientResourceAllocator.cpp b/Engine/Source/Runtime/D3D12RHI/Private/D3D12TransientResourceAllocator.cpp
--- a/Engine/Source/Runtime/D3D12RHI/Private/D3D12TransientResourceAllocator.cpp
+++ b/Engine/Source/Runtime/D3D12RHI/Private/D3D12TransientResourceAllocator.cpp
@@ -6,25 +6,6 @@
 #include "ProfilingDebugging/MemoryTrace.h"
 #include "ProfilingDebugging/AssetMetadataTrace.h"
 
-D3D12_RESOURCE_STATES GetInitialResourceState(const D3D12_RESOURCE_DESC& InDesc, ED3D12QueueType QueueType)
-{
-	// Validate the creation state
-	D3D12_RESOURCE_STATES State = D3D12_RESOURCE_STATE_COMMON;
-	if (EnumHasAnyFlags(InDesc.Flags, D3D12_RESOURCE_FLAG_ALLOW_RENDER_TARGET) && QueueType == ED3D12QueueType::Direct)
-	{
-		State = D3D12_RESOURCE_STATE_RENDER_TARGET;
-	}
-	else if (EnumHasAnyFlags(InDesc.Flags, D3D12_RESOURCE_FLAG_ALLOW_DEPTH_STENCIL) && QueueType == ED3D12QueueType::Direct)
-	{
-		State = D3D12_RESOURCE_STATE_DEPTH_WRITE;
-	}
-	else if (EnumHasAnyFlags(InDesc.Flags, D3D12_RESOURCE_FLAG_ALLOW_UNORDERED_ACCESS))
-	{
-		State = D3D12_RESOURCE_STATE_UNORDERED_ACCESS;
-	}
-	return State;
-}
-
 FD3D12TransientHeap::FD3D12TransientHeap(const FInitializer& Initializer, FD3D12Adapter* Adapter, FD3D12Device* Device, FRHIGPUMask VisibleNodeMask)
 	: FRHITransientHeap(Initializer)
 {
@@ -186,8 +167,17 @@ FRHITransientTexture* FD3D12TransientResourceHeapAllocator::CreateTexture(const
 }
 
 void FD3D12TransientResourceHeapAllocator::FResourceAllocatorAdapter::AllocateResource(
-	uint32 GPUIndex, D3D12_HEAP_TYPE, const FD3D12ResourceDesc& InDesc, uint64 InSize, uint32, ED3D12ResourceStateMode InResourceStateMode,
-	D3D12_RESOURCE_STATES InCreateState, const D3D12_CLEAR_VALUE* InClearValue, const TCHAR* InName, FD3D12ResourceLocation& ResourceLocation)
+	uint32 GPUIndex,
+	D3D12_HEAP_TYPE,
+	const FD3D12ResourceDesc& InDesc,
+	uint64 InSize,
+	uint32 /*InAllocationAlignment*/,
+	ED3D12Access InInitialD3D12Access,
+	ED3D12ResourceStateMode InResourceStateMode,
+	ED3D12Access InDefaultD3D12Access,
+	const D3D12_CLEAR_VALUE* InClearValue,
+	const TCHAR* InName,
+	FD3D12ResourceLocation& ResourceLocation)
 {
 	// The D3D12_RESOURCE_DESC's are built in two different functions right now. This checks that they actually match what we expect.
 #if DO_CHECK
@@ -200,7 +190,16 @@ void FD3D12TransientResourceHeapAllocator::FResourceAllocatorAdapter::AllocateRe
 
 	FD3D12Resource* NewResource = nullptr;
 	FD3D12Adapter* Adapter = GetParentAdapter();
-	VERIFYD3D12RESULT(Adapter->CreatePlacedResource(InDesc, Heap.GetLinkedObject(GPUIndex)->Get(), Allocation.Offset, InCreateState, InResourceStateMode, D3D12_RESOURCE_STATE_TBD, InClearValue, &NewResource, InName));
+	VERIFYD3D12RESULT(Adapter->CreatePlacedResource(
+		InDesc,
+		Heap.GetLinkedObject(GPUIndex)->Get(), 
+		Allocation.Offset,
+		InInitialD3D12Access,
+		InResourceStateMode,
+		InDefaultD3D12Access,
+		InClearValue,
+		&NewResource,
+		InName));
 
 	check(!ResourceLocation.IsValid());
 	ResourceLocation.AsHeapAliased(NewResource);
diff --git a/Engine/Source/Runtime/D3D12RHI/Private/D3D12TransientResourceAllocator.h b/Engine/Source/Runtime/D3D12RHI/Private/D3D12TransientResourceAllocator.h
--- a/Engine/Source/Runtime/D3D12RHI/Private/D3D12TransientResourceAllocator.h
+++ b/Engine/Source/Runtime/D3D12RHI/Private/D3D12TransientResourceAllocator.h
@@ -5,8 +5,6 @@
 #include "D3D12Resources.h"
 #include "RHICoreTransientResourceAllocator.h"
 
-extern D3D12_RESOURCE_STATES GetInitialResourceState(const D3D12_RESOURCE_DESC& InDesc, ED3D12QueueType QueueType);
-
 class FD3D12TransientHeap final
 	: public FRHITransientHeap
 	, public FRefCountBase
@@ -71,8 +69,17 @@ private:
 		{}
 
 		void AllocateResource(
-			uint32 GPUIndex, D3D12_HEAP_TYPE InHeapType, const FD3D12ResourceDesc& InDesc, uint64 InSize, uint32 InAllocationAlignment, ED3D12ResourceStateMode InResourceStateMode,
-			D3D12_RESOURCE_STATES InCreateState, const D3D12_CLEAR_VALUE* InClearValue, const TCHAR* InName, FD3D12ResourceLocation& ResourceLocation) override;
+			uint32 GPUIndex,
+			D3D12_HEAP_TYPE InHeapType,
+			const FD3D12ResourceDesc& InDesc,
+			uint64 InSize,
+			uint32 InAllocationAlignment,
+			ED3D12Access InInitialD3D12Access,
+			ED3D12ResourceStateMode InResourceStateMode,
+			ED3D12Access InDefaultD3D12Access,
+			const D3D12_CLEAR_VALUE* InClearValue,
+			const TCHAR* InName,
+			FD3D12ResourceLocation& ResourceLocation) override;
 
 		FD3D12TransientHeap& Heap;
 		const FRHITransientHeapAllocation& Allocation;
diff --git a/Engine/Source/Runtime/D3D12RHI/Private/D3D12UAV.cpp b/Engine/Source/Runtime/D3D12RHI/Private/D3D12UAV.cpp
--- a/Engine/Source/Runtime/D3D12RHI/Private/D3D12UAV.cpp
+++ b/Engine/Source/Runtime/D3D12RHI/Private/D3D12UAV.cpp
@@ -42,9 +42,9 @@ void FD3D12UnorderedAccessView::UpdateResourceInfo(const FResourceInfo& InResour
 		Device->GetParentAdapter()->CreateBuffer(
 			CD3DX12_HEAP_PROPERTIES(D3D12_HEAP_TYPE_DEFAULT, Node.GetNative(), Node.GetNative())
 			, Node
-			, D3D12_RESOURCE_STATE_UNORDERED_ACCESS
+			, ED3D12Access::UAVMask
 			, ED3D12ResourceStateMode::MultiState
-			, D3D12_RESOURCE_STATE_UNORDERED_ACCESS
+			, ED3D12Access::UAVMask
 			, 4
 			, CounterResource.GetInitReference()
 			, TEXT("Counter")
diff --git a/Engine/Source/Runtime/D3D12RHI/Private/D3D12Util.cpp b/Engine/Source/Runtime/D3D12RHI/Private/D3D12Util.cpp
--- a/Engine/Source/Runtime/D3D12RHI/Private/D3D12Util.cpp
+++ b/Engine/Source/Runtime/D3D12RHI/Private/D3D12Util.cpp
@@ -1636,243 +1636,6 @@ void LogExecuteCommandLists(uint32 NumCommandLists, ID3D12CommandList* const* pp
 	}
 }
 
-FString ConvertToResourceStateString(uint32 ResourceState)
-{
-	if (ResourceState == 0)
-	{
-		return TEXT("D3D12_RESOURCE_STATE_COMMON");
-	}
-
-	if (ResourceState == D3D12_RESOURCE_STATE_TBD)
-	{
-		return TEXT("D3D12_RESOURCE_STATE_TBD");
-	}
-
-	const TCHAR* ResourceStateNames[] =
-	{
-		TEXT("D3D12_RESOURCE_STATE_VERTEX_AND_CONSTANT_BUFFER"),
-		TEXT("D3D12_RESOURCE_STATE_INDEX_BUFFER"),
-		TEXT("D3D12_RESOURCE_STATE_RENDER_TARGET"),
-		TEXT("D3D12_RESOURCE_STATE_UNORDERED_ACCESS"),
-		TEXT("D3D12_RESOURCE_STATE_DEPTH_WRITE"),
-		TEXT("D3D12_RESOURCE_STATE_DEPTH_READ"),
-		TEXT("D3D12_RESOURCE_STATE_NON_PIXEL_SHADER_RESOURCE"),
-		TEXT("D3D12_RESOURCE_STATE_PIXEL_SHADER_RESOURCE"),
-		TEXT("D3D12_RESOURCE_STATE_STREAM_OUT"),
-		TEXT("D3D12_RESOURCE_STATE_INDIRECT_ARGUMENT"),
-		TEXT("D3D12_RESOURCE_STATE_COPY_DEST"),
-		TEXT("D3D12_RESOURCE_STATE_COPY_SOURCE"),
-		TEXT("D3D12_RESOURCE_STATE_RESOLVE_DEST"),
-		TEXT("D3D12_RESOURCE_STATE_RESOLVE_SOURCE"),
-	};
-
-	FString ResourceStateString;
-	uint16 NumStates = 0;
-	for (uint16 i = 0; ResourceState && i < ARRAYSIZE(ResourceStateNames); i++)
-	{
-		if (ResourceState & 1)
-		{
-			if (NumStates > 0)
-			{
-				ResourceStateString += " | ";
-			}
-
-			ResourceStateString += ResourceStateNames[i];
-			NumStates++;
-		}
-		ResourceState = ResourceState >> 1;
-	}
-	return ResourceStateString;
-}
-
-void LogResourceBarriers(TConstArrayView<D3D12_RESOURCE_BARRIER> Barriers, ID3D12CommandList* const pCommandList, ED3D12QueueType QueueType, const FString& ResourceName)
-{
-	// Configure what resource barriers are logged.
-	const bool bLogAll = true;
-	const bool bLogTransitionDepth = true;
-	const bool bLogTransitionRenderTarget = true;
-	const bool bLogTransitionUAV = true;
-	const bool bCheckResourceName = ResourceName.IsEmpty() ? false : true;
-
-	// Create the state bit mask to indicate what barriers should be logged.
-	uint32 ShouldLogMask = bLogAll ? static_cast<uint32>(-1) : 0;
-	ShouldLogMask |= bLogTransitionDepth ? D3D12_RESOURCE_STATE_DEPTH_READ | D3D12_RESOURCE_STATE_DEPTH_WRITE : 0;
-	ShouldLogMask |= bLogTransitionRenderTarget ? D3D12_RESOURCE_STATE_RENDER_TARGET : 0;
-	ShouldLogMask |= bLogTransitionUAV ? D3D12_RESOURCE_STATE_UNORDERED_ACCESS : 0;
-
-	for (int32 i = 0; i < Barriers.Num(); i++)
-	{
-		const D3D12_RESOURCE_BARRIER& currentBarrier = Barriers[i];
-
-		switch (currentBarrier.Type)
-		{
-		case D3D12_RESOURCE_BARRIER_TYPE_TRANSITION:
-		{
-			const FString StateBefore = ConvertToResourceStateString(static_cast<uint32>(currentBarrier.Transition.StateBefore));
-			const FString StateAfter = ConvertToResourceStateString(static_cast<uint32>(currentBarrier.Transition.StateAfter));
-
-			bool bShouldLog = bLogAll;
-			if (!bShouldLog)
-			{
-				// See if we should log this transition.
-				for (uint32 j = 0; (j < 2) && !bShouldLog; j++)
-				{
-					const D3D12_RESOURCE_STATES& State = (j == 0) ? currentBarrier.Transition.StateBefore : currentBarrier.Transition.StateAfter;
-					bShouldLog = (State & ShouldLogMask) > 0;
-				}
-			}
-
-			if (bShouldLog)
-			{
-				const FString BarrierResourceName = GetD312ObjectName(currentBarrier.Transition.pResource);
-
-				if ((bCheckResourceName == false) || (BarrierResourceName == ResourceName))
-				{
-					UE_LOG(LogD3D12RHI, Log, TEXT("*** BARRIER (CmdList: %016llX, Queue: %s) %u/%u: %s %016llX (Sub: %u), %s -> %s"), pCommandList, GetD3DCommandQueueTypeName(QueueType), i + 1, Barriers.Num(),
-						*BarrierResourceName,
-						currentBarrier.Transition.pResource,
-						currentBarrier.Transition.Subresource,
-						*StateBefore,
-						*StateAfter);
-				}
-			}
-			break;
-		}
-
-		case D3D12_RESOURCE_BARRIER_TYPE_UAV:
-			{
-				const FString BarrierResourceName = GetD312ObjectName(currentBarrier.UAV.pResource);
-				if ((bCheckResourceName == false) || (BarrierResourceName == ResourceName))
-				{
-					UE_LOG(LogD3D12RHI, Log, TEXT("*** BARRIER (CmdList: %016llX, Queue: %s) %u/%u: UAV Barrier %s"), pCommandList, GetD3DCommandQueueTypeName(QueueType), i + 1, Barriers.Num(), *BarrierResourceName);
-				}
-			}
-			break;
-
-		case D3D12_RESOURCE_BARRIER_TYPE_ALIASING:
-			{
-				const FString BarrierResourceNameBefore = GetD312ObjectName(currentBarrier.Aliasing.pResourceBefore);
-				const FString BarrierResourceNameAfter = GetD312ObjectName(currentBarrier.Aliasing.pResourceAfter);
-				
-				if ((bCheckResourceName == false) || (BarrierResourceNameBefore == ResourceName) || (BarrierResourceNameAfter == ResourceName))
-				{
-					UE_LOG(LogD3D12RHI, Log, TEXT("*** BARRIER (CmdList: %016llX, Queue: %s) %u/%u: Aliasing Barrier, %016llX %s -> %016llX %s"), pCommandList, GetD3DCommandQueueTypeName(QueueType), i + 1, Barriers.Num(), currentBarrier.Aliasing.pResourceBefore, *BarrierResourceNameBefore, currentBarrier.Aliasing.pResourceAfter, *BarrierResourceNameAfter);
-				}
-			}
-			break;
-
-		default:
-			check(false);
-			break;
-		}
-	}
-}
-
-
-D3D12_RESOURCE_STATES GetD3D12ResourceState(ERHIAccess InRHIAccess, ED3D12QueueType QueueType, FD3D12Texture* InRHID3D12Texture)
-{
-	// Add switch for common states (should cover all writeable states)
-	switch (InRHIAccess)
-	{
-		// all single write states
-	case ERHIAccess::RTV:					return D3D12_RESOURCE_STATE_RENDER_TARGET;
-	case ERHIAccess::UAVMask:
-	case ERHIAccess::UAVCompute:
-	case ERHIAccess::UAVGraphics:			return D3D12_RESOURCE_STATE_UNORDERED_ACCESS;
-	case ERHIAccess::DSVWrite:				return D3D12_RESOURCE_STATE_DEPTH_WRITE;
-	case ERHIAccess::CopyDest:				return D3D12_RESOURCE_STATE_COPY_DEST;
-	case ERHIAccess::ResolveDst:			return D3D12_RESOURCE_STATE_RESOLVE_DEST;
-	case ERHIAccess::Present:				return D3D12_RESOURCE_STATE_PRESENT;
-
-		// Generic read for mask read states
-	case ERHIAccess::ReadOnlyMask:
-	case ERHIAccess::ReadOnlyExclusiveMask:	return D3D12_RESOURCE_STATE_GENERIC_READ;
-	default:
-	{
-		D3D12_RESOURCE_STATES ExtraReadState = D3D12_RESOURCE_STATES(0);
-
-		if (InRHID3D12Texture)
-		{
-// 			if (InRHID3D12Texture->GetResource()->IsDepthStencilResource())
-// 			{
-// 				ExtraReadState |= D3D12_RESOURCE_STATE_DEPTH_READ;
-// 			}
-			 
-			if (InRHID3D12Texture->SkipsFastClearFinalize())
-			{
-				ExtraReadState |= FD3D12CommandContext::SkipFastClearEliminateState;
-			}
-		}
-
-		// Special case for DSV read & write (Depth write allows depth read as well in D3D)
-		if (InRHIAccess == ERHIAccess(ERHIAccess::DSVRead | ERHIAccess::DSVWrite))
-		{
-			return D3D12_RESOURCE_STATE_DEPTH_WRITE;
-		}
-		else
-		{
-			// Should be combination from read only flags (write flags covered above)
-			check(!(EnumHasAnyFlags(InRHIAccess, ERHIAccess::WritableMask)));
-			check(EnumHasAnyFlags(InRHIAccess, ERHIAccess::ReadOnlyMask));
-
-			D3D12_RESOURCE_STATES State = D3D12_RESOURCE_STATE_COMMON;
-
-			// Translate the requested after state to a D3D state
-			if (EnumHasAnyFlags(InRHIAccess, ERHIAccess::SRVGraphics) && QueueType == ED3D12QueueType::Direct)
-			{
-				State |= D3D12_RESOURCE_STATE_PIXEL_SHADER_RESOURCE | D3D12_RESOURCE_STATE_NON_PIXEL_SHADER_RESOURCE | ExtraReadState;
-			}
-			if (EnumHasAnyFlags(InRHIAccess, ERHIAccess::SRVCompute))
-			{
-				State |= D3D12_RESOURCE_STATE_NON_PIXEL_SHADER_RESOURCE | ExtraReadState;
-			}
-			if (EnumHasAnyFlags(InRHIAccess, ERHIAccess::VertexOrIndexBuffer))
-			{
-				State |= D3D12_RESOURCE_STATE_VERTEX_AND_CONSTANT_BUFFER | D3D12_RESOURCE_STATE_INDEX_BUFFER;
-			}
-			if (EnumHasAnyFlags(InRHIAccess, ERHIAccess::CopySrc))
-			{
-				State |= D3D12_RESOURCE_STATE_COPY_SOURCE;
-			}
-			if (EnumHasAnyFlags(InRHIAccess, ERHIAccess::IndirectArgs))
-			{
-				State |= D3D12_RESOURCE_STATE_INDIRECT_ARGUMENT;
-			}
-			if (EnumHasAnyFlags(InRHIAccess, ERHIAccess::ResolveSrc))
-			{
-				State |= D3D12_RESOURCE_STATE_RESOLVE_SOURCE;
-			}
-			if (EnumHasAnyFlags(InRHIAccess, ERHIAccess::DSVRead))
-			{
-				State |= D3D12_RESOURCE_STATE_DEPTH_READ;
-			}
-			if (EnumHasAnyFlags(InRHIAccess, ERHIAccess::ShadingRateSource))
-			{
-#if !UE_BUILD_SHIPPING
-				if (GRHISupportsAttachmentVariableRateShading == false)
-				{
-					static bool bLogOnce = true;
-					if (bLogOnce)
-					{
-						UE_LOG(LogD3D12RHI, Warning, TEXT("(%s) Resource state is D3D12_RESOURCE_STATE_SHADING_RATE_SOURCE but RHI does not support VRS."), InRHID3D12Texture == nullptr ? TEXT("Unknown") : *InRHID3D12Texture->GetName().GetPlainNameString());
-					}
-					bLogOnce = false;
-				}
-#endif
-
-#if PLATFORM_SUPPORTS_VARIABLE_RATE_SHADING
-				State |= D3D12_RESOURCE_STATE_SHADING_RATE_SOURCE;
-#endif
-			}
-
-			// Should have at least one valid state
-			check(State != D3D12_RESOURCE_STATE_COMMON);
-
-			return State;
-		}
-	}
-	}
-}
 
 #if ASSERT_RESOURCE_STATES
 // Forward declarations are required for the template functions
diff --git a/Engine/Source/Runtime/D3D12RHI/Private/D3D12Util.h b/Engine/Source/Runtime/D3D12RHI/Private/D3D12Util.h
--- a/Engine/Source/Runtime/D3D12RHI/Private/D3D12Util.h
+++ b/Engine/Source/Runtime/D3D12RHI/Private/D3D12Util.h
@@ -10,6 +10,7 @@
 #include "Containers/Queue.h"
 #include "D3D12RHICommon.h"
 #include "D3D12Queue.h"
+#include "D3D12Access.h"
 #include "DXGIUtilities.h"
 #include "RenderUtils.h"
 #include "ShaderCore.h"
@@ -265,27 +266,6 @@ private:
 };
 
 void LogExecuteCommandLists(uint32 NumCommandLists, ID3D12CommandList *const *ppCommandLists);
-FString ConvertToResourceStateString(uint32 ResourceState);
-void LogResourceBarriers(TConstArrayView<D3D12_RESOURCE_BARRIER> Barriers, ID3D12CommandList *const pCommandList, ED3D12QueueType QueueType, const FString& ResourceName);
-
-// Custom resource states
-// To Be Determined (TBD) means we need to fill out a resource barrier before the command list is executed.
-#define D3D12_RESOURCE_STATE_TBD D3D12_RESOURCE_STATES(-1 ^ (1 << 31))
-#define D3D12_RESOURCE_STATE_CORRUPT D3D12_RESOURCE_STATES(-2 ^ (1 << 31))
-
-static bool IsValidD3D12ResourceState(D3D12_RESOURCE_STATES InState)
-{
-	return (InState != D3D12_RESOURCE_STATE_TBD && InState != D3D12_RESOURCE_STATE_CORRUPT);
-}
-
-static bool IsDirectQueueExclusiveD3D12State(D3D12_RESOURCE_STATES InState)
-{
-	return EnumHasAnyFlags(InState, D3D12_RESOURCE_STATE_RENDER_TARGET | D3D12_RESOURCE_STATE_DEPTH_WRITE | D3D12_RESOURCE_STATE_DEPTH_READ | D3D12_RESOURCE_STATE_PIXEL_SHADER_RESOURCE);
-}
-
-class FD3D12Texture;
-
-D3D12_RESOURCE_STATES GetD3D12ResourceState(ERHIAccess InRHIAccess, ED3D12QueueType QueueType, FD3D12Texture* InRHID3D12Texture);
 
 /**
  * The base class of threadsafe reference counted objects.
@@ -402,16 +382,16 @@ inline bool IsCPUAccessible(D3D12_HEAP_TYPE HeapType, const D3D12_HEAP_PROPERTIE
 	return !IsGPUOnly(HeapType, pCustomHeapProperties);
 }
 
-inline D3D12_RESOURCE_STATES DetermineInitialResourceState(D3D12_HEAP_TYPE HeapType, const D3D12_HEAP_PROPERTIES *pCustomHeapProperties = nullptr)
+inline ED3D12Access DetermineInitialBufferD3D12Access(D3D12_HEAP_TYPE HeapType, const D3D12_HEAP_PROPERTIES* pCustomHeapProperties = nullptr)
 {
 	if (HeapType == D3D12_HEAP_TYPE_DEFAULT || IsCPUWritable(HeapType, pCustomHeapProperties))
 	{
-		return D3D12_RESOURCE_STATE_GENERIC_READ;
+		return ED3D12Access::GenericRead;
 	}
 	else
 	{
 		check(HeapType == D3D12_HEAP_TYPE_READBACK);
-		return D3D12_RESOURCE_STATE_COPY_DEST;
+		return ED3D12Access::CopyDest;
 	}
 }
 
diff --git a/Engine/Source/Runtime/D3D12RHI/Private/D3D12Viewport.cpp b/Engine/Source/Runtime/D3D12RHI/Private/D3D12Viewport.cpp
--- a/Engine/Source/Runtime/D3D12RHI/Private/D3D12Viewport.cpp
+++ b/Engine/Source/Runtime/D3D12RHI/Private/D3D12Viewport.cpp
@@ -86,7 +86,13 @@ FD3D12Texture* GetSwapChainSurface(FD3D12Device* Parent, EPixelFormat PixelForma
 		TextureDesc.Layout = D3D12_TEXTURE_LAYOUT_UNKNOWN;
 		TextureDesc.Flags = D3D12_RESOURCE_FLAG_ALLOW_RENDER_TARGET;
 
-		Parent->GetDevice()->CreateCommittedResource(&HeapProps, D3D12_HEAP_FLAG_NONE, &TextureDesc, D3D12_RESOURCE_STATE_PRESENT, nullptr, IID_PPV_ARGS(BackBufferResource.GetInitReference()));
+		Parent->GetParentAdapter()->CreateCommittedResourceRaw(
+			HeapProps,
+			D3D12_HEAP_FLAG_NONE,
+			TextureDesc,
+			ED3D12Access::Present,
+			nullptr,
+			BackBufferResource);
 	}
 
 	FD3D12ResourceDesc BackBufferDesc = BackBufferResource->GetDesc();
@@ -123,11 +129,15 @@ FD3D12Texture* GetSwapChainSurface(FD3D12Device* Parent, EPixelFormat PixelForma
 	{
 		FD3D12Texture* NewTexture = DynamicRHI->CreateNewD3D12Texture(CreateDesc, Device);
 
-		const D3D12_RESOURCE_STATES InitialState = D3D12_RESOURCE_STATE_COMMON;
-
 		if (Device->GetGPUIndex() == Parent->GetGPUIndex())
 		{
-			FD3D12Resource* NewResourceWrapper = new FD3D12Resource(Device, FRHIGPUMask::All(), BackBufferResource, InitialState, BackBufferDesc);
+			FD3D12Resource* NewResourceWrapper =
+				new FD3D12Resource(
+					Device,
+					FRHIGPUMask::All(),
+					BackBufferResource,
+					ED3D12Access::Common,
+					BackBufferDesc);
 			NewResourceWrapper->AddRef();
 			NewTexture->ResourceLocation.AsStandAlone(NewResourceWrapper);
 		}
@@ -142,7 +152,8 @@ FD3D12Texture* GetSwapChainSurface(FD3D12Device* Parent, EPixelFormat PixelForma
 				NewTexture,
 				PixelFormat,
 				TexCreate_RenderTargetable | TexCreate_ShaderResource | TexCreate_Presentable,
-				InitialState,
+				ED3D12Access::Common,
+				ED3D12Access::Common,
 				*Name);
 		}
 
diff --git a/Engine/Source/Runtime/D3D12RHI/Private/D3D12WorkGraph.cpp b/Engine/Source/Runtime/D3D12RHI/Private/D3D12WorkGraph.cpp
--- a/Engine/Source/Runtime/D3D12RHI/Private/D3D12WorkGraph.cpp
+++ b/Engine/Source/Runtime/D3D12RHI/Private/D3D12WorkGraph.cpp
@@ -983,12 +983,12 @@ void FD3D12CommandContext::DispatchWorkGraphShaderBundle(
 			16,
 			CreateDesc,
 			ED3D12ResourceStateMode::MultiState,
-			D3D12_RESOURCE_STATE_COPY_DEST,
+			ED3D12Access::CopyDest,
 			true
 		);
 
 		BatchedSyncPoints.ToWait.Emplace(RootArgBuffer->UploadResourceDataViaCopyQueue(*this, &LocalRootArgs));
-		TransitionResource(RootArgBuffer->GetResource(), D3D12_RESOURCE_STATE_COPY_DEST, D3D12_RESOURCE_STATE_COMMON, 0);
+		AddBarrier(RootArgBuffer->GetResource(), ED3D12Access::CopyDest, ED3D12Access::Common, D3D12_RESOURCE_BARRIER_ALL_SUBRESOURCES);
 
 		NodeLocalRootArgumentsTable.StartAddress = RootArgBuffer->ResourceLocation.GetGPUVirtualAddress();
 		NodeLocalRootArgumentsTable.SizeInBytes = RootArgBuffer->ResourceLocation.GetSize();
@@ -1342,12 +1342,12 @@ void FD3D12CommandContext::DispatchWorkGraphShaderBundle(
 			16,
 			CreateDesc,
 			ED3D12ResourceStateMode::MultiState,
-			D3D12_RESOURCE_STATE_COPY_DEST,
+			ED3D12Access::CopyDest,
 			true
 		);
 
 		BatchedSyncPoints.ToWait.Emplace(RootArgBuffer->UploadResourceDataViaCopyQueue(*this, &LocalRootArgs));
-		TransitionResource(RootArgBuffer->GetResource(), D3D12_RESOURCE_STATE_COPY_DEST, D3D12_RESOURCE_STATE_COMMON, 0);
+		AddBarrier(RootArgBuffer->GetResource(), ED3D12Access::CopyDest, ED3D12Access::Common, D3D12_RESOURCE_BARRIER_ALL_SUBRESOURCES);
 
 		NodeLocalRootArgumentsTable.StartAddress = RootArgBuffer->ResourceLocation.GetGPUVirtualAddress();
 		NodeLocalRootArgumentsTable.SizeInBytes = RootArgBuffer->ResourceLocation.GetSize();
diff --git a/Engine/Source/Runtime/D3D12RHI/Private/Windows/WindowsD3D12Device.cpp b/Engine/Source/Runtime/D3D12RHI/Private/Windows/WindowsD3D12Device.cpp
--- a/Engine/Source/Runtime/D3D12RHI/Private/Windows/WindowsD3D12Device.cpp
+++ b/Engine/Source/Runtime/D3D12RHI/Private/Windows/WindowsD3D12Device.cpp
@@ -2214,7 +2214,16 @@ FD3D12DiagnosticBuffer::FD3D12DiagnosticBuffer(FD3D12Queue& Queue)
 	Heap = new FD3D12Heap(Queue.Device, Queue.Device->GetGPUMask());
 	Heap->SetHeap(D3D12Heap, TEXT("DiagnosticBuffer"));
 
-	hr = Queue.Device->GetParentAdapter()->CreatePlacedResource(Desc, Heap.GetReference(), 0, D3D12_RESOURCE_STATE_COPY_DEST, nullptr, Resource.GetInitReference(), *Name, false);
+	hr = Queue.Device->GetParentAdapter()->CreatePlacedResource(
+		Desc,
+		Heap.GetReference(),
+		0,
+		ED3D12Access::CopyDest,
+		nullptr,
+		Resource.GetInitReference(),
+		*Name,
+		false);
+
 	if (FAILED(hr))
 	{
 		Heap.SafeRelease();
diff --git a/Engine/Source/Runtime/D3D12RHI/Private/Windows/WindowsD3D12RHIDefinitions.h b/Engine/Source/Runtime/D3D12RHI/Private/Windows/WindowsD3D12RHIDefinitions.h
--- a/Engine/Source/Runtime/D3D12RHI/Private/Windows/WindowsD3D12RHIDefinitions.h
+++ b/Engine/Source/Runtime/D3D12RHI/Private/Windows/WindowsD3D12RHIDefinitions.h
@@ -33,3 +33,6 @@ typedef ID3D12GraphicsCommandList ID3D12CopyCommandList;
 
 // Only enable pipeline statistics if we've got the stats system enabled to display them
 #define D3D12RHI_ENABLE_PIPELINE_STATISTICS (1 && STATS)
+
+#define D3D12RHI_ALLOW_ENHANCED_BARRIERS 1
+#define D3D12RHI_ALLOW_LEGACY_BARRIERS   1
diff --git a/Engine/Source/Runtime/RHI/Public/RHIContext.h b/Engine/Source/Runtime/RHI/Public/RHIContext.h
--- a/Engine/Source/Runtime/RHI/Public/RHIContext.h
+++ b/Engine/Source/Runtime/RHI/Public/RHIContext.h
@@ -527,7 +527,8 @@ public:
 	inline ERHIPipeline GetTrackedPipelines(const FRHIViewableResource* Resource) const
 	{
 		check(Resource);
-		return Resource->TrackedAccess.Pipelines != ERHIPipeline::None ? Resource->TrackedAccess.Pipelines : GetPipeline();
+		// The assumption is that all resources begin life on the gfx pipe
+		return Resource->TrackedAccess.Pipelines != ERHIPipeline::None ? Resource->TrackedAccess.Pipelines : ERHIPipeline::Graphics;
 	}
 
 	virtual void* RHIGetNativeCommandBuffer() { return nullptr; }
diff --git a/Engine/Source/Runtime/RHICore/Internal/RHICoreTransitions.h b/Engine/Source/Runtime/RHICore/Internal/RHICoreTransitions.h
--- a/Engine/Source/Runtime/RHICore/Internal/RHICoreTransitions.h
+++ b/Engine/Source/Runtime/RHICore/Internal/RHICoreTransitions.h
@@ -18,7 +18,7 @@ struct FResourceState
 	{
 		if (FRHIViewableResource* ViewableResource = GetViewableResource(Info))
 		{
-			if (AccessBefore == ERHIAccess::Unknown || AccessBefore == ERHIAccess::Discard)
+			if (NeedToReplaceBeforeInfo())
 			{
 				SrcPipelines = Context.GetTrackedPipelines(ViewableResource);
 				AccessBefore = Context.GetTrackedAccess(ViewableResource);
@@ -40,6 +40,21 @@ struct FResourceState
 	// SrcPipeline may differ from the declared RHI SrcPipeline when CollapseToSinglePipeline is used, in which case this will be ERHIPipeline::All.
 	ERHIPipeline SrcPipelines;
 	ERHIPipeline DstPipelines;
+
+private:
+	bool NeedToReplaceBeforeInfo() const
+	{
+		if (AccessBefore == ERHIAccess::Unknown)
+		{
+			return true;
+		}
+		if (AccessBefore == ERHIAccess::Discard
+			&& GRHIGlobals.NeedsTransientDiscardStateTracking)
+		{
+			return true;
+		}
+		return false;
+	}
 };
 
 } //! UE::RHICore
\ No newline at end of file
