Create Resources With Correct Alignment Value in D3D12PoolAllocator

This was a regression introduced by CL 44756831.

#rnx

diff --git a/Engine/Source/Runtime/D3D12RHI/Private/D3D12PoolAllocator.cpp b/Engine/Source/Runtime/D3D12RHI/Private/D3D12PoolAllocator.cpp
--- a/Engine/Source/Runtime/D3D12RHI/Private/D3D12PoolAllocator.cpp
+++ b/Engine/Source/Runtime/D3D12RHI/Private/D3D12PoolAllocator.cpp
@@ -442,7 +442,7 @@ void FD3D12PoolAllocator::AllocateResource(
 			FD3D12Resource* NewResource =
 				CreatePlacedResource(
 					AllocationData,
-					InDesc,
+					Desc,
 					InInitialD3D12Access,
 					InResourceStateMode,
 					InDefaultD3D12Access,
