Suppress PVS Warning in Enhanced Barriers

PVS complains that the value D3D12_BARRIER_SYNC_VERTEX_SHADING shows up twice
in the same list of OR'd together enum values. It's right, but this table is also meant to
match the official spec and that value shows up twice there too.

#rnx

diff --git a/Engine/Source/Runtime/D3D12RHI/Private/D3D12EnhancedBarriers.cpp b/Engine/Source/Runtime/D3D12RHI/Private/D3D12EnhancedBarriers.cpp
--- a/Engine/Source/Runtime/D3D12RHI/Private/D3D12EnhancedBarriers.cpp
+++ b/Engine/Source/Runtime/D3D12RHI/Private/D3D12EnhancedBarriers.cpp
@@ -425,7 +425,7 @@ TStaticArray<D3D12_BARRIER_SYNC, 23> AccessCompatibleSync = {
 	| D3D12_BARRIER_SYNC_RENDER_TARGET,
 
 	//D3D12_BARRIER_ACCESS_UNORDERED_ACCESS
-	D3D12_BARRIER_SYNC_ALL
+	D3D12_BARRIER_SYNC_ALL //-V501
 	| D3D12_BARRIER_SYNC_VERTEX_SHADING
 	| D3D12_BARRIER_SYNC_PIXEL_SHADING
 	| D3D12_BARRIER_SYNC_COMPUTE_SHADING
