Fix Invalid D3D12Access Combination in Raytracing

#rnx

diff --git a/Engine/Source/Runtime/D3D12RHI/Private/D3D12RayTracing.cpp b/Engine/Source/Runtime/D3D12RHI/Private/D3D12RayTracing.cpp
--- a/Engine/Source/Runtime/D3D12RHI/Private/D3D12RayTracing.cpp
+++ b/Engine/Source/Runtime/D3D12RHI/Private/D3D12RayTracing.cpp
@@ -5533,7 +5533,7 @@ static void DispatchRays(FD3D12CommandContext& CommandContext,
 
 		CommandContext.AddBarrier(
 			ArgumentBuffer->GetResource(),
-			ED3D12Access::UAVCompute | ED3D12Access::IndirectArgs,
+			ED3D12Access::SRVGraphicsNonPixel | ED3D12Access::IndirectArgs,
 			ED3D12Access::CopySrc,
 			0);
 
@@ -5579,7 +5579,7 @@ static void DispatchRays(FD3D12CommandContext& CommandContext,
 		CommandContext.AddBarrier(
 			ArgumentBuffer->GetResource(),
 			ED3D12Access::CopySrc,
-			ED3D12Access::UAVCompute | ED3D12Access::IndirectArgs,
+			ED3D12Access::SRVGraphicsNonPixel | ED3D12Access::IndirectArgs,
 			0);
 
 		CommandContext.FlushResourceBarriers();
