Add missing D3D capability check for D3D12 GPU breadcrumbs.

Check for D3D12_FEATURE_EXISTING_HEAPS when creating a D3D12 diagnostic buffer, and do not attempt to use WriteBufferImmediate if the diagnostic buffer is unsupported.

diff --git a/Engine/Source/Runtime/D3D12RHI/Private/D3D12CommandContext.cpp b/Engine/Source/Runtime/D3D12RHI/Private/D3D12CommandContext.cpp
--- a/Engine/Source/Runtime/D3D12RHI/Private/D3D12CommandContext.cpp
+++ b/Engine/Source/Runtime/D3D12RHI/Private/D3D12CommandContext.cpp
@@ -97,7 +97,8 @@ void FD3D12ContextCommon::BindDiagnosticBuffer(FD3D12RootSignature const* RootSi
 	if (Slot < 0)
 		return;
 
-	if (FD3D12DiagnosticBuffer* DiagBuffer = Device->GetQueue(QueueType).DiagnosticBuffer.Get())
+	FD3D12DiagnosticBuffer* DiagBuffer = Device->GetQueue(QueueType).DiagnosticBuffer.Get();
+	if (DiagBuffer && DiagBuffer->IsValid())
 	{
 		D3D12_GPU_VIRTUAL_ADDRESS DataAddress = DiagBuffer->GetGPUQueueData();
 
@@ -114,7 +115,7 @@ void FD3D12ContextCommon::BindDiagnosticBuffer(FD3D12RootSignature const* RootSi
 	void FD3D12CommandContext::RHIBeginBreadcrumbGPU(FRHIBreadcrumbNode* Breadcrumb)
 	{
 		FD3D12DiagnosticBuffer* DiagBuffer = Device->GetQueue(QueueType).DiagnosticBuffer.Get();
-		if (DiagBuffer && UE::RHI::UseGPUCrashBreadcrumbs())
+		if (DiagBuffer && DiagBuffer->IsValid() && UE::RHI::UseGPUCrashBreadcrumbs())
 		{
 			D3D12_GPU_VIRTUAL_ADDRESS Marker = DiagBuffer->GetGPUQueueMarkerIn();
 			WriteMarker(Marker, Breadcrumb->ID, EMarkerType::In);
@@ -229,7 +230,7 @@ void FD3D12ContextCommon::BindDiagnosticBuffer(FD3D12RootSignature const* RootSi
 	#endif
 
 		FD3D12DiagnosticBuffer* DiagBuffer = Device->GetQueue(QueueType).DiagnosticBuffer.Get();
-		if (DiagBuffer && UE::RHI::UseGPUCrashBreadcrumbs())
+		if (DiagBuffer && DiagBuffer->IsValid() && UE::RHI::UseGPUCrashBreadcrumbs())
 		{
 			D3D12_GPU_VIRTUAL_ADDRESS Marker = DiagBuffer->GetGPUQueueMarkerOut();
 			WriteMarker(Marker, Breadcrumb->ID, EMarkerType::Out);
diff --git a/Engine/Source/Runtime/D3D12RHI/Private/Windows/WindowsD3D12Device.cpp b/Engine/Source/Runtime/D3D12RHI/Private/Windows/WindowsD3D12Device.cpp
--- a/Engine/Source/Runtime/D3D12RHI/Private/Windows/WindowsD3D12Device.cpp
+++ b/Engine/Source/Runtime/D3D12RHI/Private/Windows/WindowsD3D12Device.cpp
@@ -2190,6 +2190,13 @@ FD3D12DiagnosticBuffer::FD3D12DiagnosticBuffer(FD3D12Queue& Queue)
 		return;
 	}
 
+	D3D12_FEATURE_DATA_EXISTING_HEAPS ExistingHeapSupport{};
+	if (FAILED(D3D12Device3->CheckFeatureSupport(D3D12_FEATURE_EXISTING_HEAPS, &ExistingHeapSupport, sizeof(ExistingHeapSupport))) || !ExistingHeapSupport.Supported)
+	{
+		UE_LOG(LogD3D12RHI, Warning, TEXT("[GPUBreadCrumb] D3D12_FEATURE_EXISTING_HEAPS is not supported."));
+		return;
+	}
+
 	// Allocate persistent CPU readable memory which will still be valid after a device lost and wrap this data in a placed resource so the GPU command list can write to it
 	Data = static_cast<FQueue*>(VirtualAlloc(nullptr, Desc.Width, MEM_COMMIT | MEM_RESERVE, PAGE_READWRITE));
 	if (!Data)
