Fix Clang 20 Compilation Errors In D3D12 Barriers

Clang 20 complains "constexpr variable 'BarrierType_<Name>' must be initialized by a constant expression" so we just won't make them constexpr.
Also remove the macros that were previously handling silencing related warnings.

#rnx

diff --git a/Engine/Source/Runtime/D3D12RHI/Private/D3D12EnhancedBarriers.cpp b/Engine/Source/Runtime/D3D12RHI/Private/D3D12EnhancedBarriers.cpp
--- a/Engine/Source/Runtime/D3D12RHI/Private/D3D12EnhancedBarriers.cpp
+++ b/Engine/Source/Runtime/D3D12RHI/Private/D3D12EnhancedBarriers.cpp
@@ -1752,13 +1752,11 @@ static void LogBarriers(
 
 class FD3D12EnhancedBarriersBatcher
 {
-	PRAGMA_DISABLE_ENUM_CONSTEXPR_CONVERSION_WARNINGS
 	// Use the top bit of the flags enum to mark transitions as "idle" time (used to remove the swapchain wait time for back buffers).
-	static constexpr D3D12_BARRIER_TYPE BarrierType_CountAsIdleTime = D3D12_BARRIER_TYPE(1ull << ((sizeof(D3D12_BARRIER_TYPE) * 8) - 1));
-	static constexpr D3D12_BARRIER_TYPE D3D12_BARRIER_TYPE_GLOBAL_COUNT_AS_IDLE  = D3D12_BARRIER_TYPE(D3D12_BARRIER_TYPE_GLOBAL  | BarrierType_CountAsIdleTime);
-	static constexpr D3D12_BARRIER_TYPE D3D12_BARRIER_TYPE_TEXTURE_COUNT_AS_IDLE = D3D12_BARRIER_TYPE(D3D12_BARRIER_TYPE_TEXTURE | BarrierType_CountAsIdleTime);
-	static constexpr D3D12_BARRIER_TYPE D3D12_BARRIER_TYPE_BUFFER_COUNT_AS_IDLE  = D3D12_BARRIER_TYPE(D3D12_BARRIER_TYPE_BUFFER  | BarrierType_CountAsIdleTime);
-	PRAGMA_ENABLE_ENUM_CONSTEXPR_CONVERSION_WARNINGS
+	static const D3D12_BARRIER_TYPE BarrierType_CountAsIdleTime = D3D12_BARRIER_TYPE(1ull << ((sizeof(D3D12_BARRIER_TYPE) * 8) - 1));
+	static const D3D12_BARRIER_TYPE D3D12_BARRIER_TYPE_GLOBAL_COUNT_AS_IDLE  = D3D12_BARRIER_TYPE(D3D12_BARRIER_TYPE_GLOBAL  | BarrierType_CountAsIdleTime);
+	static const D3D12_BARRIER_TYPE D3D12_BARRIER_TYPE_TEXTURE_COUNT_AS_IDLE = D3D12_BARRIER_TYPE(D3D12_BARRIER_TYPE_TEXTURE | BarrierType_CountAsIdleTime);
+	static const D3D12_BARRIER_TYPE D3D12_BARRIER_TYPE_BUFFER_COUNT_AS_IDLE  = D3D12_BARRIER_TYPE(D3D12_BARRIER_TYPE_BUFFER  | BarrierType_CountAsIdleTime);
 
 	struct FD3D12BarrierGroup : public D3D12_BARRIER_GROUP
 	{
diff --git a/Engine/Source/Runtime/D3D12RHI/Private/D3D12LegacyBarriers.cpp b/Engine/Source/Runtime/D3D12RHI/Private/D3D12LegacyBarriers.cpp
--- a/Engine/Source/Runtime/D3D12RHI/Private/D3D12LegacyBarriers.cpp
+++ b/Engine/Source/Runtime/D3D12RHI/Private/D3D12LegacyBarriers.cpp
@@ -739,10 +739,8 @@ const TCHAR* FD3D12LegacyBarriersForAdapter::GetImplementationName() const
 
 class FD3D12LegacyBarriersBatcher
 {
-	PRAGMA_DISABLE_ENUM_CONSTEXPR_CONVERSION_WARNINGS
 	// Use the top bit of the flags enum to mark transitions as "idle" time (used to remove the swapchain wait time for back buffers).
 	static const D3D12_RESOURCE_BARRIER_FLAGS BarrierFlag_CountAsIdleTime = D3D12_RESOURCE_BARRIER_FLAGS(1ull << ((sizeof(D3D12_RESOURCE_BARRIER_FLAGS) * 8) - 1));
-	PRAGMA_ENABLE_ENUM_CONSTEXPR_CONVERSION_WARNINGS
 
 	struct FD3D12ResourceBarrier : public D3D12_RESOURCE_BARRIER
 	{
