Minor adjustments in RHI alignment tests

diff --git a/Engine/Plugins/Tests/RHITests/Source/RHITests/Private/RHIAllocatorTests.cpp b/Engine/Plugins/Tests/RHITests/Source/RHITests/Private/RHIAllocatorTests.cpp
--- a/Engine/Plugins/Tests/RHITests/Source/RHITests/Private/RHIAllocatorTests.cpp
+++ b/Engine/Plugins/Tests/RHITests/Source/RHITests/Private/RHIAllocatorTests.cpp
@@ -34,7 +34,6 @@ bool FRHIAllocatorTests::Test_LockBuffer16ByteAlignment(FRHICommandListImmediate
 		const FBufferRHIRef TestBuffer = RHICmdList.CreateBuffer(TestBufferCreateDesc);
 
 		uintptr_t MappedAddr = (uintptr_t)RHICmdList.LockBuffer(TestBuffer, 0, sizeof(FVector4f) * 5, RLM_WriteOnly);
-		check(MappedAddr % ExpectedOutAlignment == 0);
 		bool bResult = (MappedAddr % ExpectedOutAlignment) == 0;
 
 		if (!bResult)
diff --git a/Engine/Source/Runtime/D3D12RHI/Private/D3D12Allocation.cpp b/Engine/Source/Runtime/D3D12RHI/Private/D3D12Allocation.cpp
--- a/Engine/Source/Runtime/D3D12RHI/Private/D3D12Allocation.cpp
+++ b/Engine/Source/Runtime/D3D12RHI/Private/D3D12Allocation.cpp
@@ -426,10 +426,10 @@ void FD3D12BuddyAllocator::Allocate(uint32 SizeInBytes, uint32 Alignment, FD3D12
 	uint32 SizeToAllocate = SizeInBytes;
 
 	// Need to always align to a multiple of 16
-	Alignment = Align(Alignment, 16u);
+	Alignment = Align(FMath::Max(Alignment, 16u), 16u);
 
 	// If the alignment doesn't match the block size
-	if (Alignment != 0 && MinBlockSize % Alignment != 0)
+	if (MinBlockSize % Alignment != 0)
 	{
 		SizeToAllocate = SizeInBytes + Alignment;
 	}
@@ -443,7 +443,7 @@ void FD3D12BuddyAllocator::Allocate(uint32 SizeInBytes, uint32 Alignment, FD3D12
 	const uint32 AllocationBlockOffset = uint32(Offset * MinBlockSize);
 	uint32 Padding = 0;
 
-	if (Alignment != 0 && AllocationBlockOffset % Alignment != 0)
+	if (AllocationBlockOffset % Alignment != 0)
 	{
 		uint32 AlignedBlockOffset = AlignArbitrary(AllocationBlockOffset, Alignment);
 		Padding = AlignedBlockOffset - AllocationBlockOffset;
@@ -901,7 +901,7 @@ bool FD3D12BucketAllocator::TryAllocate(uint32 SizeInBytes, uint32 Alignment, FD
 	FD3D12Adapter* Adapter = GetParentDevice()->GetParentAdapter();
 
 	// Alignment needs to be a multiple of 16 bytes
-	Alignment = Align(Alignment, 16u);
+	Alignment = Align(FMath::Max(Alignment, 16u), 16u);
 
 	// Size cannot be smaller than the requested alignment
 	SizeInBytes = FMath::Max(SizeInBytes, Alignment);
@@ -975,7 +975,7 @@ bool FD3D12BucketAllocator::TryAllocate(uint32 SizeInBytes, uint32 Alignment, FD
 	}
 
 	uint64 AlignedBlockOffset = Block.Offset;
-	if (Alignment != 0 && AlignedBlockOffset % Alignment != 0)
+	if (AlignedBlockOffset % Alignment != 0)
 	{
 		AlignedBlockOffset = AlignArbitrary(AlignedBlockOffset, Alignment);
 	}
