RHIResourceCollection Updates
- Adding update functions
- Adding Bindless RHI unit tests. Currently just a route to make sure the ResourceCollection calls don't fail, full tests to come later.

diff --git a/Engine/Plugins/Tests/RHITests/Source/RHITests/Private/RHIBindlessTests.cpp b/Engine/Plugins/Tests/RHITests/Source/RHITests/Private/RHIBindlessTests.cpp
new file mode 100644
--- /dev/null
+++ b/Engine/Plugins/Tests/RHITests/Source/RHITests/Private/RHIBindlessTests.cpp
@@ -0,0 +1,65 @@
+// Copyright Epic Games, Inc. All Rights Reserved.
+
+#include "RHIBindlessTests.h"
+#include "RHIResourceUtils.h"
+
+bool RHIBindlessTests::Test_ResourceCollection(FRHICommandListImmediate& RHICmdList)
+{
+	if (!GRHIGlobals.bSupportsBindless)
+	{
+		return true;
+	}
+
+	const uint32 BufferContents[] = { 1, 2, 3, 4 };
+	const uint32 TextureContents[] = { 1 };
+
+	const FRHIBufferCreateDesc BufferDesc =
+		FRHIBufferCreateDesc::CreateByteAddress(TEXT("ResourceCollection_Buffer"), sizeof(BufferContents), sizeof(BufferContents[0]))
+		.SetInitialState(ERHIAccess::SRVMask)
+		.SetInitActionInitializer();
+	TRHIBufferInitializer<uint32> BufferInitializer = RHICmdList.CreateBufferInitializer(BufferDesc);
+	{
+		BufferInitializer.WriteArray(MakeConstArrayView(BufferContents));
+	}
+	FBufferRHIRef TestBuffer = BufferInitializer.Finalize();
+
+	FRHIViewDesc::FBufferSRV::FInitializer BufferSRVDesc = FRHIViewDesc::CreateBufferSRV().SetType(FRHIViewDesc::EBufferType::Raw);
+	FShaderResourceViewRHIRef BufferSRV = RHICmdList.CreateShaderResourceView(TestBuffer, BufferSRVDesc);
+
+	const FRHITextureCreateDesc TextureDesc =
+		FRHITextureCreateDesc::Create2D(TEXT("ResourceCollection_Texture"), 1, 1, PF_R32_UINT)
+		.SetInitialState(ERHIAccess::SRVMask)
+		.SetInitActionInitializer();
+	FRHITextureInitializer TextureInitializer = RHICmdList.CreateTextureInitializer(TextureDesc);
+	{
+		TextureInitializer.GetTexture2DSubresource(0).WriteData(&TextureContents[0], sizeof(TextureContents));
+	}
+	FTextureRHIRef TestTexture = TextureInitializer.Finalize();
+
+	// Empty
+	{
+		FRHIResourceCollectionRef EmptyCollection = RHICmdList.CreateResourceCollection({});
+	}
+
+	// Normal Creation
+	{
+		TArray<FRHIResourceCollectionMember> Members;
+		Members.Emplace(BufferSRV.GetReference());
+		Members.Emplace(TestTexture.GetReference());
+
+		FRHIResourceCollectionRef ResourceCollection = RHICmdList.CreateResourceCollection(Members);
+	}
+
+	// Updates
+	{
+		TArray<FRHIResourceCollectionMember> Members;
+		Members.Emplace(BufferSRV.GetReference());
+		Members.Emplace(TestTexture.GetReference());
+
+		FRHIResourceCollectionRef ResourceCollection = RHICmdList.CreateResourceCollection(Members);
+
+		RHICmdList.UpdateResourceCollection(ResourceCollection, 0, Members);
+	}
+
+	return true;
+}
diff --git a/Engine/Plugins/Tests/RHITests/Source/RHITests/Private/RHIBindlessTests.h b/Engine/Plugins/Tests/RHITests/Source/RHITests/Private/RHIBindlessTests.h
new file mode 100644
--- /dev/null
+++ b/Engine/Plugins/Tests/RHITests/Source/RHITests/Private/RHIBindlessTests.h
@@ -0,0 +1,10 @@
+// Copyright Epic Games, Inc. All Rights Reserved.
+
+#pragma once
+
+class FRHICommandListImmediate;
+
+namespace RHIBindlessTests
+{
+	bool Test_ResourceCollection(FRHICommandListImmediate& RHICmdList);
+}
diff --git a/Engine/Plugins/Tests/RHITests/Source/RHITests/Private/RHITests.spec.cpp b/Engine/Plugins/Tests/RHITests/Source/RHITests/Private/RHITests.spec.cpp
--- a/Engine/Plugins/Tests/RHITests/Source/RHITests/Private/RHITests.spec.cpp
+++ b/Engine/Plugins/Tests/RHITests/Source/RHITests/Private/RHITests.spec.cpp
@@ -12,6 +12,7 @@
 #include "RHIReservedResourceTests.h"
 #include "RHIGraphicsUAVTests.h"
 #include "RHIAllocatorTests.h"
+#include "RHIBindlessTests.h"
 
 BEGIN_DEFINE_SPEC(FAutomationRHITest, "Rendering.RHI", EAutomationTestFlags::EngineFilter | EAutomationTestFlags_ApplicationContextMask | EAutomationTestFlags::NonNullRHI)
 END_DEFINE_SPEC(FAutomationRHITest)
@@ -211,4 +212,12 @@ void FAutomationRHITest::Define()
 		});
 	});
 
+	Describe("Test RHI Bindless", [this]()
+	{
+		It("RHI Resource Collection", [this]()
+		{
+			bool bResult = RunOnRenderThreadSynchronous(RHIBindlessTests::Test_ResourceCollection);
+			TestEqual("Resource Collection failed", bResult, 1);
+		});
+	});
 }
diff --git a/Engine/Plugins/Tests/RHITests/Source/RHITests/Private/RHITestsModule.cpp b/Engine/Plugins/Tests/RHITests/Source/RHITests/Private/RHITestsModule.cpp
--- a/Engine/Plugins/Tests/RHITests/Source/RHITests/Private/RHITestsModule.cpp
+++ b/Engine/Plugins/Tests/RHITests/Source/RHITests/Private/RHITestsModule.cpp
@@ -13,6 +13,7 @@
 #include "RHIReservedResourceTests.h"
 #include "RHIGraphicsUAVTests.h"
 #include "RHIAllocatorTests.h"
+#include "RHIBindlessTests.h"
 
 #define LOCTEXT_NAMESPACE "FRHITestsModule"
 
@@ -112,6 +113,13 @@ static bool RunTests_RenderThread(FRHICommandListImmediate& RHICmdList)
 		RUN_TEST(FRHIAllocatorTests::Test_LockBuffer16ByteAlignment(RHICmdList));
 	}
 
+	// ------------------------------------------------
+	// Bindless
+	// ------------------------------------------------
+	{
+		RUN_TEST(RHIBindlessTests::Test_ResourceCollection(RHICmdList));
+	}
+
 	// @todo - add more tests
 	return bResult;
 }
diff --git a/Engine/Source/Runtime/Apple/MetalRHI/Private/MetalDynamicRHI.h b/Engine/Source/Runtime/Apple/MetalRHI/Private/MetalDynamicRHI.h
--- a/Engine/Source/Runtime/Apple/MetalRHI/Private/MetalDynamicRHI.h
+++ b/Engine/Source/Runtime/Apple/MetalRHI/Private/MetalDynamicRHI.h
@@ -137,6 +137,7 @@ public:
 
 #if PLATFORM_SUPPORTS_BINDLESS_RENDERING
 	virtual FRHIResourceCollectionRef RHICreateResourceCollection(FRHICommandListBase& RHICmdList, TConstArrayView<FRHIResourceCollectionMember> InMembers) final override;
+	virtual void RHIUpdateResourceCollection(FRHICommandListBase& RHICmdList, FRHIResourceCollection* InResourceCollection, uint32 InStartIndex, TConstArrayView<FRHIResourceCollectionMember> InMemberUpdates) final override;
 #endif
 
 	virtual FRHICalcTextureSizeResult RHICalcTexturePlatformSize(FRHITextureDesc const& Desc, uint32 FirstMipIndex) final override;
diff --git a/Engine/Source/Runtime/Apple/MetalRHI/Private/MetalRHI.cpp b/Engine/Source/Runtime/Apple/MetalRHI/Private/MetalRHI.cpp
--- a/Engine/Source/Runtime/Apple/MetalRHI/Private/MetalRHI.cpp
+++ b/Engine/Source/Runtime/Apple/MetalRHI/Private/MetalRHI.cpp
@@ -1842,7 +1842,12 @@ void FMetalDynamicRHI::RHIReplaceResources(FRHICommandListBase& RHICmdList, TArr
 #if PLATFORM_SUPPORTS_BINDLESS_RENDERING
 FRHIResourceCollectionRef FMetalDynamicRHI::RHICreateResourceCollection(FRHICommandListBase& RHICmdList, TConstArrayView<FRHIResourceCollectionMember> InMembers)
 {
-	return new FMetalResourceCollection(RHICmdList, InMembers);
+	return UE::RHICore::CreateGenericResourceCollection(RHICmdList, InMembers);
+}
+
+void FMetalDynamicRHI::RHIUpdateResourceCollection(FRHICommandListBase& RHICmdList, FRHIResourceCollection* InResourceCollection, uint32 InStartIndex, TConstArrayView<FRHIResourceCollectionMember> InMemberUpdates)
+{
+	UE::RHICore::UpdateGenericResourceCollection(RHICmdList, ResourceCast(InResourceCollection), InStartIndex, InMemberUpdates);
 }
 #endif
 
diff --git a/Engine/Source/Runtime/D3D12RHI/Private/D3D12RHIPrivate.h b/Engine/Source/Runtime/D3D12RHI/Private/D3D12RHIPrivate.h
--- a/Engine/Source/Runtime/D3D12RHI/Private/D3D12RHIPrivate.h
+++ b/Engine/Source/Runtime/D3D12RHI/Private/D3D12RHIPrivate.h
@@ -443,6 +443,7 @@ public:
 	virtual FUnorderedAccessViewRHIRef RHICreateUnorderedAccessView(class FRHICommandListBase& RHICmdList, FRHIViewableResource* Resource, FRHIViewDesc const& ViewDesc) override;
 #if PLATFORM_SUPPORTS_BINDLESS_RENDERING
 	virtual FRHIResourceCollectionRef RHICreateResourceCollection(FRHICommandListBase& RHICmdList, TConstArrayView<FRHIResourceCollectionMember> InMembers) override;
+	virtual void RHIUpdateResourceCollection(FRHICommandListBase& RHICmdList, FRHIResourceCollection* InResourceCollection, uint32 InStartIndex, TConstArrayView<FRHIResourceCollectionMember> InMemberUpdates) override;
 #endif
 
 	// ID3D12DynamicRHI interface.
diff --git a/Engine/Source/Runtime/D3D12RHI/Private/D3D12ResourceCollection.cpp b/Engine/Source/Runtime/D3D12RHI/Private/D3D12ResourceCollection.cpp
--- a/Engine/Source/Runtime/D3D12RHI/Private/D3D12ResourceCollection.cpp
+++ b/Engine/Source/Runtime/D3D12RHI/Private/D3D12ResourceCollection.cpp
@@ -8,7 +8,7 @@
 #include "D3D12CommandContext.h"
 #include "D3D12TextureReference.h"
 
-FD3D12ResourceCollection::FD3D12ResourceCollection(FD3D12Device* InParent, FRHICommandListBase& RHICmdList, FD3D12Buffer* InBuffer, TConstArrayView<FRHIResourceCollectionMember> InMembers, FD3D12ResourceCollection* FirstLinkedObject)
+FD3D12ResourceCollection::FD3D12ResourceCollection(FD3D12Device* InParent,FD3D12Buffer* InBuffer, TConstArrayView<FRHIResourceCollectionMember> InMembers, FD3D12ResourceCollection* FirstLinkedObject)
 	: FRHIResourceCollection(InMembers)
 	, FD3D12DeviceChild(InParent)
 	, Buffer(InBuffer->GetLinkedObject(InParent->GetGPUIndex()))
@@ -75,9 +75,67 @@ FRHIResourceCollectionRef FD3D12DynamicRHI::RHICreateResourceCollection(FRHIComm
 	FRHIViewDesc::FBufferSRV::FInitializer ViewDesc = FRHIViewDesc::CreateBufferSRV().SetType(FRHIViewDesc::EBufferType::Raw);
 	FShaderResourceViewRHIRef ShaderResourceView = RHICmdList.CreateShaderResourceView(Buffer, ViewDesc);
 
-	return GetAdapter().CreateLinkedObject<FD3D12ResourceCollection>(FRHIGPUMask::All(), [&RHICmdList, Buffer, InMembers](FD3D12Device* Device, FD3D12ResourceCollection* FirstLinkedObject)
+	return GetAdapter().CreateLinkedObject<FD3D12ResourceCollection>(FRHIGPUMask::All(), [Buffer, InMembers](FD3D12Device* Device, FD3D12ResourceCollection* FirstLinkedObject)
 	{
-		return new FD3D12ResourceCollection(Device, RHICmdList, Buffer, InMembers, FirstLinkedObject);
+		return new FD3D12ResourceCollection(Device, Buffer, InMembers, FirstLinkedObject);
+	});
+}
+
+void FD3D12DynamicRHI::RHIUpdateResourceCollection(FRHICommandListBase& RHICmdList, FRHIResourceCollection* InResourceCollection, uint32 InStartIndex, TConstArrayView<FRHIResourceCollectionMember> InMemberUpdates)
+{
+	TConstArrayView<FRHIResourceCollectionMember> MemberUpdates = UE::RHICore::GetValidResourceCollectionUpdateList(InResourceCollection, InStartIndex, InMemberUpdates);
+	if (MemberUpdates.IsEmpty())
+	{
+		return;
+	}
+
+	FD3D12ResourceCollection* ResourceCollection = ResourceCast(InResourceCollection);
+	TConstArrayView<FRHIResourceCollectionMember> CommandMemberUpdates = RHICmdList.IsTopOfPipe() ? RHICmdList.AllocArray(MemberUpdates) : MemberUpdates;
+
+	RHICmdList.EnqueueLambda([ResourceCollection, InStartIndex, CommandMemberUpdates](FRHICommandListBase& RHICmdList)
+	{
+		const uint32 UploadSize = CommandMemberUpdates.Num() * sizeof(uint32);
+		const uint32 DestOffset = (InStartIndex + 1) * sizeof(uint32);
+
+		ResourceCollection->UpdateMembers(InStartIndex, CommandMemberUpdates);
+
+		for (FD3D12Buffer::FLinkedObjectIterator BufferIt(ResourceCollection->Buffer); BufferIt; ++BufferIt)
+		{
+			FD3D12Buffer& Buffer = *BufferIt;
+			FD3D12Device* Device = Buffer.GetParentDevice();
+			FD3D12UploadHeapAllocator& Allocator = Device->GetParentAdapter()->GetUploadHeapAllocator(Device->GetGPUIndex());
+
+			FD3D12ResourceLocation UploadResourceLocation(Device);
+			uint32* UploadData = reinterpret_cast<uint32*>(Allocator.AllocUploadResource(UploadSize, D3D12_TEXTURE_DATA_PLACEMENT_ALIGNMENT, UploadResourceLocation));
+
+			UE::RHICore::FillResourceCollectionUpdateMemory(UploadData, CommandMemberUpdates);
+
+			FD3D12CommandContext& Context = FD3D12CommandContext::Get(RHICmdList, Device->GetGPUIndex());
+
+			FD3D12Resource* SourceResource = UploadResourceLocation.GetResource();
+			FD3D12Resource* DestResource = Buffer.ResourceLocation.GetResource();
+
+			// Clear the resource if still bound to make sure the SRVs are rebound again on next operation (and get correct resource transitions enqueued)
+			Context.ConditionalClearShaderResource(&Buffer.ResourceLocation, EShaderParameterTypeMask::SRVMask);
+
+			FScopedResourceBarrier ScopeResourceBarrierDest(Context, DestResource, ED3D12Access::CopyDest, 0);
+
+			// Don't need to transition upload heaps
+			Context.FlushResourceBarriers();
+
+			Context.UpdateResidency(DestResource);
+			Context.UpdateResidency(SourceResource);
+
+			Context.CopyBufferRegionChecked(
+				DestResource->GetResource(), DestResource->GetName(),
+				Buffer.ResourceLocation.GetOffsetFromBaseOfResource() + DestOffset,
+				SourceResource->GetResource(), SourceResource->GetName(),
+				UploadResourceLocation.GetOffsetFromBaseOfResource(),
+				UploadSize
+			);
+
+			Context.ConditionalSplitCommandList();
+		}
 	});
 }
 
diff --git a/Engine/Source/Runtime/D3D12RHI/Private/D3D12ResourceCollection.h b/Engine/Source/Runtime/D3D12RHI/Private/D3D12ResourceCollection.h
--- a/Engine/Source/Runtime/D3D12RHI/Private/D3D12ResourceCollection.h
+++ b/Engine/Source/Runtime/D3D12RHI/Private/D3D12ResourceCollection.h
@@ -16,7 +16,7 @@ class FD3D12ShaderResourceView;
 class FD3D12ResourceCollection : public FRHIResourceCollection, public FD3D12DeviceChild, public FD3D12LinkedAdapterObject<FD3D12ResourceCollection>
 {
 public:
-	FD3D12ResourceCollection(FD3D12Device* InParent, FRHICommandListBase& RHICmdList, FD3D12Buffer* InBuffer, TConstArrayView<FRHIResourceCollectionMember> InMembers, FD3D12ResourceCollection* FirstLinkedObject);
+	FD3D12ResourceCollection(FD3D12Device* InParent, FD3D12Buffer* InBuffer, TConstArrayView<FRHIResourceCollectionMember> InMembers, FD3D12ResourceCollection* FirstLinkedObject);
 	~FD3D12ResourceCollection();
 
 	virtual FRHIDescriptorHandle GetBindlessHandle() const final;
diff --git a/Engine/Source/Runtime/RHI/Private/RHIResourceCollection.cpp b/Engine/Source/Runtime/RHI/Private/RHIResourceCollection.cpp
new file mode 100644
--- /dev/null
+++ b/Engine/Source/Runtime/RHI/Private/RHIResourceCollection.cpp
@@ -0,0 +1,63 @@
+// Copyright Epic Games, Inc. All Rights Reserved.
+
+#include "RHIResourceCollection.h"
+
+
+FRHIResourceCollection::FRHIResourceCollection(TConstArrayView<FRHIResourceCollectionMember> InMembers)
+	: FRHIResource(RRT_ResourceCollection)
+	, Members(InMembers)
+{
+	for (const FRHIResourceCollectionMember& Member : Members)
+	{
+		if (Member.Resource)
+		{
+			Member.Resource->AddRef();
+		}
+	}
+}
+
+FRHIResourceCollection::~FRHIResourceCollection()
+{
+	for (const FRHIResourceCollectionMember& Member : Members)
+	{
+		if (Member.Resource)
+		{
+			Member.Resource->Release();
+		}
+	}
+}
+
+void FRHIResourceCollection::UpdateMember(int32 Index, FRHIResourceCollectionMember InMember)
+{
+	if (Index >= 0 && Index < Members.Num())
+	{
+		FRHIResourceCollectionMember& ExistingMember = Members[Index];
+		if (ExistingMember.Resource != InMember.Resource)
+		{
+			if (InMember.Resource)
+			{
+				InMember.Resource->AddRef();
+			}
+
+			if (ExistingMember.Resource)
+			{
+				ExistingMember.Resource->Release();
+			}
+
+			ExistingMember = InMember;
+		}
+	}
+}
+
+void FRHIResourceCollection::UpdateMembers(int32 StartIndex, TConstArrayView<FRHIResourceCollectionMember> NewMembers)
+{
+	for (int32 NewMemberIndex = 0; NewMemberIndex < NewMembers.Num(); NewMemberIndex++)
+	{
+		UpdateMember(StartIndex + NewMemberIndex, NewMembers[NewMemberIndex]);
+	}
+}
+
+FRHIDescriptorHandle FRHIResourceCollection::GetBindlessHandle() const
+{
+	return FRHIDescriptorHandle();
+}
diff --git a/Engine/Source/Runtime/RHI/Private/RHIValidation.cpp b/Engine/Source/Runtime/RHI/Private/RHIValidation.cpp
--- a/Engine/Source/Runtime/RHI/Private/RHIValidation.cpp
+++ b/Engine/Source/Runtime/RHI/Private/RHIValidation.cpp
@@ -420,7 +420,7 @@ namespace RHIValidation
 
 			void SetResourceCollection(FRHIResourceCollection* ResourceCollection, uint8 Index)
 			{
-				for (const FRHIResourceCollectionMember& Member : ResourceCollection->Members)
+				for (const FRHIResourceCollectionMember& Member : ResourceCollection->GetMembers())
 				{
 					switch (Member.Type)
 					{
@@ -3666,7 +3666,7 @@ void ValidateShaderParameters(FRHIShader* RHIShader, RHIValidation::FTracker* Tr
 		case FRHIShaderParameterResource::EType::ResourceCollection:
 			if (const FRHIResourceCollection* ResourceCollection = static_cast<const FRHIResourceCollection*>(Parameter.Resource))
 			{
-				for (const FRHIResourceCollectionMember& Member : ResourceCollection->Members)
+				for (const FRHIResourceCollectionMember& Member : ResourceCollection->GetMembers())
 				{
 					switch (Member.Type)
 					{
diff --git a/Engine/Source/Runtime/RHI/Public/DynamicRHI.h b/Engine/Source/Runtime/RHI/Public/DynamicRHI.h
--- a/Engine/Source/Runtime/RHI/Public/DynamicRHI.h
+++ b/Engine/Source/Runtime/RHI/Public/DynamicRHI.h
@@ -531,6 +531,10 @@ public:
 		return FRHIResourceCollectionRef{};
 	}
 
+	virtual void RHIUpdateResourceCollection(FRHICommandListBase& RHICmdList, FRHIResourceCollection* InResourceCollection, uint32 InStartIndex, TConstArrayView<FRHIResourceCollectionMember> InMemberUpdates)
+	{
+	}
+
 	/**
 	* Computes the size in memory required by a given texture.
 	*
diff --git a/Engine/Source/Runtime/RHI/Public/RHICommandList.h b/Engine/Source/Runtime/RHI/Public/RHICommandList.h
--- a/Engine/Source/Runtime/RHI/Public/RHICommandList.h
+++ b/Engine/Source/Runtime/RHI/Public/RHICommandList.h
@@ -1147,6 +1147,11 @@ PRAGMA_ENABLE_DEPRECATION_WARNINGS
 		return GDynamicRHI->RHICreateResourceCollection(*this, InMembers);
 	}
 
+	inline void UpdateResourceCollection(FRHIResourceCollection* InResourceCollection, uint32 InStartIndex, TConstArrayView<FRHIResourceCollectionMember> InMemberUpdates)
+	{
+		return GDynamicRHI->RHIUpdateResourceCollection(*this, InResourceCollection, InStartIndex, InMemberUpdates);
+	}
+
 	inline FRayTracingGeometryRHIRef CreateRayTracingGeometry(const FRayTracingGeometryInitializer& Initializer)
 	{
 		return GDynamicRHI->RHICreateRayTracingGeometry(*this, Initializer);
diff --git a/Engine/Source/Runtime/RHI/Public/RHIResourceCollection.h b/Engine/Source/Runtime/RHI/Public/RHIResourceCollection.h
--- a/Engine/Source/Runtime/RHI/Public/RHIResourceCollection.h
+++ b/Engine/Source/Runtime/RHI/Public/RHIResourceCollection.h
@@ -45,34 +45,20 @@ struct FRHIResourceCollectionMember
 class FRHIResourceCollection : public FRHIResource
 {
 public:
-	FRHIResourceCollection(TConstArrayView<FRHIResourceCollectionMember> InMembers)
-		: FRHIResource(RRT_ResourceCollection)
-		, Members(InMembers)
-	{
-		for (const FRHIResourceCollectionMember& Member : Members)
-		{
-			if (Member.Resource)
-			{
-				Member.Resource->AddRef();
-			}
-		}
-	}
+	RHI_API FRHIResourceCollection(TConstArrayView<FRHIResourceCollectionMember> InMembers);
 	FRHIResourceCollection(const FRHIResourceCollection&) = delete;
-	~FRHIResourceCollection()
-	{
-		for (const FRHIResourceCollectionMember& Member : Members)
-		{
-			if (Member.Resource)
-			{
-				Member.Resource->Release();
-			}
-		}
-	}
+	RHI_API ~FRHIResourceCollection();
+
+	RHI_API void UpdateMember(int32 Index, FRHIResourceCollectionMember InMember);
+	RHI_API void UpdateMembers(int32 StartIndex, TConstArrayView<FRHIResourceCollectionMember> NewMembers);
+
+	RHI_API virtual FRHIDescriptorHandle GetBindlessHandle() const;
 
-	virtual FRHIDescriptorHandle GetBindlessHandle() const
+	TConstArrayView<FRHIResourceCollectionMember> GetMembers() const
 	{
-		return FRHIDescriptorHandle();
+		return MakeConstArrayView(Members);
 	}
 
+private:
 	TArray<FRHIResourceCollectionMember> Members;
 };
diff --git a/Engine/Source/Runtime/RHI/Public/RHIValidation.h b/Engine/Source/Runtime/RHI/Public/RHIValidation.h
--- a/Engine/Source/Runtime/RHI/Public/RHIValidation.h
+++ b/Engine/Source/Runtime/RHI/Public/RHIValidation.h
@@ -456,11 +456,19 @@ public:
 	virtual FShaderResourceViewRHIRef RHICreateShaderResourceView(class FRHICommandListBase& RHICmdList, FRHIViewableResource* Resource, FRHIViewDesc const& ViewDesc) override final;
 	virtual FUnorderedAccessViewRHIRef RHICreateUnorderedAccessView(class FRHICommandListBase& RHICmdList, FRHIViewableResource* Resource, FRHIViewDesc const& ViewDesc) override final;
 
-	virtual FRHIResourceCollectionRef RHICreateResourceCollection(FRHICommandListBase& RHICmdList, TConstArrayView<FRHIResourceCollectionMember> InMembers)
+	virtual FRHIResourceCollectionRef RHICreateResourceCollection(FRHICommandListBase& RHICmdList, TConstArrayView<FRHIResourceCollectionMember> InMembers) override final
 	{
 		return RHI->RHICreateResourceCollection(RHICmdList, InMembers);
 	}
 
+	virtual void RHIUpdateResourceCollection(FRHICommandListBase& RHICmdList, FRHIResourceCollection* InResourceCollection, uint32 InStartIndex, TConstArrayView<FRHIResourceCollectionMember> InMemberUpdates) override final
+	{
+		RHI_VALIDATION_CHECK(InStartIndex < static_cast<uint32>(InResourceCollection->GetMembers().Num()) && InStartIndex + InMemberUpdates.Num() <= static_cast<uint32>(InResourceCollection->GetMembers().Num()),
+			*FString::Printf(TEXT("ResourceCollection update range of [%d, %d) is invalid for a ResourceCollection of size '%d'"), InStartIndex, InStartIndex + InMemberUpdates.Num(), InResourceCollection->GetMembers().Num()));
+
+		RHI->RHIUpdateResourceCollection(RHICmdList, InResourceCollection, InStartIndex, InMemberUpdates);
+	}
+
 	virtual FRHICalcTextureSizeResult RHICalcTexturePlatformSize(FRHITextureDesc const& Desc, uint32 FirstMipIndex) override final
 	{
 		ensure(Desc.IsValid());
diff --git a/Engine/Source/Runtime/RHICore/Internal/RHICoreResourceCollection.h b/Engine/Source/Runtime/RHICore/Internal/RHICoreResourceCollection.h
--- a/Engine/Source/Runtime/RHICore/Internal/RHICoreResourceCollection.h
+++ b/Engine/Source/Runtime/RHICore/Internal/RHICoreResourceCollection.h
@@ -38,13 +38,10 @@ namespace UE::RHICore
 		return Handle;
 	}
 
-	inline void FillResourceCollectionMemory(TRHIBufferInitializer<uint32>& Destination, TConstArrayView<FRHIResourceCollectionMember> InValues)
+	inline void FillResourceCollectionUpdateMemory(uint32* Destination, TConstArrayView<FRHIResourceCollectionMember> InValues)
 	{
 		int32 WriteIndex = 0;
 
-		Destination[WriteIndex] = static_cast<uint32>(InValues.Num());
-		++WriteIndex;
-
 		for (const FRHIResourceCollectionMember& Value : InValues)
 		{
 			const FRHIDescriptorHandle Handle = GetHandleForResourceCollectionValue(Value);
@@ -57,6 +54,13 @@ namespace UE::RHICore
 		}
 	}
 
+	inline void FillResourceCollectionMemory(TRHIBufferInitializer<uint32>& Destination, TConstArrayView<FRHIResourceCollectionMember> InValues)
+	{
+		Destination[0] = static_cast<uint32>(InValues.Num());
+
+		FillResourceCollectionUpdateMemory(Destination.GetWritableData() + 1, InValues);
+	}
+
 	inline FRHIBuffer* CreateResourceCollectionBuffer(FRHICommandListBase& RHICmdList, TConstArrayView<FRHIResourceCollectionMember> InMembers)
 	{
 		const size_t BufferSize = UE::RHICore::CalculateResourceCollectionMemorySize(InMembers);
@@ -106,6 +110,41 @@ namespace UE::RHICore
 	{
 		return new FGenericResourceCollection(RHICmdList, InMembers);
 	}
+
+	inline TConstArrayView<FRHIResourceCollectionMember> GetValidResourceCollectionUpdateList(FRHIResourceCollection* ResourceCollection, uint32 StartIndex, TConstArrayView<FRHIResourceCollectionMember> InUpdates)
+	{
+		const uint32 CollectionSize = ResourceCollection->GetMembers().Num();
+		if (StartIndex >= CollectionSize)
+		{
+			// Invalid start index
+			return {};
+		}
+
+		const int32 UpdateCount = FMath::Min<int32>(InUpdates.Num(), CollectionSize - StartIndex);
+		return InUpdates.Left(UpdateCount);
+	}
+
+	inline void UpdateGenericResourceCollection(FRHICommandListBase& RHICmdList, FGenericResourceCollection* ResourceCollection, uint32 StartIndex, TConstArrayView<FRHIResourceCollectionMember> InMemberUpdates)
+	{
+		TConstArrayView<FRHIResourceCollectionMember> MemberUpdates = UE::RHICore::GetValidResourceCollectionUpdateList(ResourceCollection, StartIndex, InMemberUpdates);
+		if (!MemberUpdates.IsEmpty())
+		{
+			const uint32 UploadSize = MemberUpdates.Num() * sizeof(uint32);
+			const uint32 DestOffset = (StartIndex + 1) * sizeof(uint32);
+
+			uint32* UploadData = reinterpret_cast<uint32*>(RHICmdList.LockBuffer(ResourceCollection->Buffer, DestOffset, UploadSize, RLM_WriteOnly));
+
+			UE::RHICore::FillResourceCollectionUpdateMemory(UploadData, MemberUpdates);
+
+			RHICmdList.UnlockBuffer(ResourceCollection->Buffer);
+
+			TConstArrayView<FRHIResourceCollectionMember> CommandMemberUpdates = RHICmdList.IsTopOfPipe() ? RHICmdList.AllocArray(MemberUpdates) : MemberUpdates;
+			RHICmdList.EnqueueLambda([ResourceCollection, StartIndex, CommandMemberUpdates](FRHICommandListBase& RHICmdList)
+			{
+				ResourceCollection->UpdateMembers(StartIndex, CommandMemberUpdates);
+			});
+		}
+	}
 }
 
 #endif // PLATFORM_SUPPORTS_BINDLESS_RENDERING
diff --git a/Engine/Source/Runtime/VulkanRHI/Private/VulkanDynamicRHI.h b/Engine/Source/Runtime/VulkanRHI/Private/VulkanDynamicRHI.h
--- a/Engine/Source/Runtime/VulkanRHI/Private/VulkanDynamicRHI.h
+++ b/Engine/Source/Runtime/VulkanRHI/Private/VulkanDynamicRHI.h
@@ -260,7 +260,8 @@ public:
 	virtual FUnorderedAccessViewRHIRef RHICreateUnorderedAccessView(class FRHICommandListBase& RHICmdList, FRHIViewableResource* Resource, FRHIViewDesc const& ViewDesc) final override;
 
 #if PLATFORM_SUPPORTS_BINDLESS_RENDERING
-	virtual FRHIResourceCollectionRef RHICreateResourceCollection(FRHICommandListBase& RHICmdList, TConstArrayView<FRHIResourceCollectionMember> InMembers) final;
+	virtual FRHIResourceCollectionRef RHICreateResourceCollection(FRHICommandListBase& RHICmdList, TConstArrayView<FRHIResourceCollectionMember> InMembers) final override;
+	virtual void RHIUpdateResourceCollection(FRHICommandListBase& RHICmdList, FRHIResourceCollection* InResourceCollection, uint32 InStartIndex, TConstArrayView<FRHIResourceCollectionMember> InMemberUpdates) final override;
 #endif
 
 	virtual FRayTracingAccelerationStructureSize RHICalcRayTracingSceneSize(const FRayTracingSceneInitializer& Initializer) final override;
diff --git a/Engine/Source/Runtime/VulkanRHI/Private/VulkanRHI.cpp b/Engine/Source/Runtime/VulkanRHI/Private/VulkanRHI.cpp
--- a/Engine/Source/Runtime/VulkanRHI/Private/VulkanRHI.cpp
+++ b/Engine/Source/Runtime/VulkanRHI/Private/VulkanRHI.cpp
@@ -2156,7 +2156,12 @@ void FVulkanDynamicRHI::RHIReplaceResources(FRHICommandListBase& RHICmdList, TAr
 #if PLATFORM_SUPPORTS_BINDLESS_RENDERING
 FRHIResourceCollectionRef FVulkanDynamicRHI::RHICreateResourceCollection(FRHICommandListBase& RHICmdList, TConstArrayView<FRHIResourceCollectionMember> InMembers)
 {
-	return new FVulkanResourceCollection(RHICmdList, InMembers);
+	return UE::RHICore::CreateGenericResourceCollection(RHICmdList, InMembers);
+}
+
+void FVulkanDynamicRHI::RHIUpdateResourceCollection(FRHICommandListBase& RHICmdList, FRHIResourceCollection* InResourceCollection, uint32 InStartIndex, TConstArrayView<FRHIResourceCollectionMember> InMemberUpdates)
+{
+	UE::RHICore::UpdateGenericResourceCollection(RHICmdList, ResourceCast(InResourceCollection), InStartIndex, InMemberUpdates);
 }
 #endif
 
