Fix usages of TArray::Emplace in D3D12RHI on types that can't be moved trivially

diff --git a/Engine/Source/Runtime/D3D12RHI/Private/D3D12Descriptors.h b/Engine/Source/Runtime/D3D12RHI/Private/D3D12Descriptors.h
--- a/Engine/Source/Runtime/D3D12RHI/Private/D3D12Descriptors.h
+++ b/Engine/Source/Runtime/D3D12RHI/Private/D3D12Descriptors.h
@@ -314,7 +314,7 @@ public:
 private:
 	TRefCountPtr<ID3D12DescriptorHeap> AcquirePooledHeap(ERHIDescriptorHeapType InType, uint32 InNumDescriptors, ED3D12DescriptorHeapFlags InFlags);
 
-	TArray<FD3D12DescriptorManager> GlobalHeaps;
+	TArray<FD3D12DescriptorManager, TFixedAllocator<2>> GlobalHeaps; // Use TFixedAllocator to prevent move of CRITICAL_SECTION
 
 	struct FPooledHeap
 	{
diff --git a/Engine/Source/Runtime/D3D12RHI/Private/D3D12Device.h b/Engine/Source/Runtime/D3D12RHI/Private/D3D12Device.h
--- a/Engine/Source/Runtime/D3D12RHI/Private/D3D12Device.h
+++ b/Engine/Source/Runtime/D3D12RHI/Private/D3D12Device.h
@@ -336,7 +336,7 @@ private:
 	FD3D12BindlessDescriptorAllocator& BindlessDescriptorAllocator;
 	FD3D12BindlessDescriptorManager BindlessDescriptorManager;
 #endif
-	TArray<FD3D12OfflineDescriptorManager, TInlineAllocator<(uint32)ERHIDescriptorHeapType::Count>> OfflineDescriptorManagers;
+	TArray<FD3D12OfflineDescriptorManager, TFixedAllocator<(uint32)ERHIDescriptorHeapType::Count>> OfflineDescriptorManagers;
 
 	FD3D12GlobalOnlineSamplerHeap GlobalSamplerHeap;
 	FD3D12OnlineDescriptorManager OnlineDescriptorManager;
@@ -347,7 +347,7 @@ private:
 	
 	FD3D12CommandContext* ImmediateCommandContext = nullptr;
 
-	TArray<FD3D12Queue, TInlineAllocator<(uint32)ED3D12QueueType::Count>> Queues;
+	TArray<FD3D12Queue, TFixedAllocator<(uint32)ED3D12QueueType::Count>> Queues;
 
 	TLruCache<D3D12_SAMPLER_DESC, TRefCountPtr<FD3D12SamplerState>> SamplerCache;
 	uint32 SamplerID = 0;
