Fix Enhanced Barriers Assert in Debug Builds

A checkSlow() path can end up asserting because it can pass ACCESS_NO_ACCESS
to GetAccessCompatibleSync(). Without the additional check, we index beyond the
end of the AccessCompatibleSync array.

#rnx

diff --git a/Engine/Source/Runtime/D3D12RHI/Private/D3D12EnhancedBarriers.cpp b/Engine/Source/Runtime/D3D12RHI/Private/D3D12EnhancedBarriers.cpp
--- a/Engine/Source/Runtime/D3D12RHI/Private/D3D12EnhancedBarriers.cpp
+++ b/Engine/Source/Runtime/D3D12RHI/Private/D3D12EnhancedBarriers.cpp
@@ -991,7 +991,8 @@ D3D12_BARRIER_SYNC GetAccessCompatibleSync(
 		std::make_unsigned_t<
 			std::underlying_type_t<D3D12_BARRIER_SYNC>>;
 
-	if (InAccess == D3D12_BARRIER_ACCESS_COMMON)
+	if (InAccess == D3D12_BARRIER_ACCESS_COMMON
+		|| InAccess == D3D12_BARRIER_ACCESS_NO_ACCESS)
 	{
 		return BitCast<D3D12_BARRIER_SYNC>(~AccessUnsignedUnderlyingType(0));
 	}
