Ensured that shared resources created as committed resources for improved compatibilty
[RN] Fixed issue that prevented PostParty from functioning on some GPUs

diff --git a/Engine/Source/Runtime/D3D12RHI/Private/D3D12PoolAllocator.cpp b/Engine/Source/Runtime/D3D12RHI/Private/D3D12PoolAllocator.cpp
--- a/Engine/Source/Runtime/D3D12RHI/Private/D3D12PoolAllocator.cpp
+++ b/Engine/Source/Runtime/D3D12RHI/Private/D3D12PoolAllocator.cpp
@@ -461,7 +461,7 @@ void FD3D12PoolAllocator::AllocateResource(
 		const D3D12_HEAP_PROPERTIES HeapProps = CD3DX12_HEAP_PROPERTIES(InHeapType, GetGPUMask().GetNative(), GetVisibilityMask().GetNative());
 
 		// If we are tracking all allocation data and allocating a standalone texture, then first create a heap so we can retrieve the GPU virtual address as well
-		if (Adapter->IsTrackingAllAllocations() && (InDesc.Dimension != D3D12_RESOURCE_DIMENSION_BUFFER))
+		if (Adapter->IsTrackingAllAllocations() && (InDesc.Dimension != D3D12_RESOURCE_DIMENSION_BUFFER) && !bSharedResource)
 		{
 			D3D12_HEAP_DESC HeapDesc = {};
 			HeapDesc.SizeInBytes = FMath::Max((uint64)D3D12_DEFAULT_RESOURCE_PLACEMENT_ALIGNMENT, InSize);
