Fixing Bindless support in D3D12 shader bundles.

diff --git a/Engine/Source/Runtime/D3D12RHI/Private/D3D12Adapter.h b/Engine/Source/Runtime/D3D12RHI/Private/D3D12Adapter.h
--- a/Engine/Source/Runtime/D3D12RHI/Private/D3D12Adapter.h
+++ b/Engine/Source/Runtime/D3D12RHI/Private/D3D12Adapter.h
@@ -238,6 +238,7 @@ public:
 
 	const FD3D12RootSignature* GetRootSignature(const FBoundShaderStateInput& BoundShaderState);
 	const FD3D12RootSignature* GetRootSignature(const class FD3D12ComputeShader* Shader);
+	const FD3D12RootSignature* GetRootSignature(ERHIShaderBundleMode ShaderBundleMode, bool bBindless);
 	
 	const FD3D12RootSignature* GetRootSignature(const class FD3D12WorkGraphShader* Shader);
 	const FD3D12RootSignature* GetGlobalWorkGraphRootSignature(const FRHIShaderBindingLayout& ShaderBindingLayout);
@@ -246,18 +247,6 @@ public:
 	const FD3D12RootSignature* GetLocalRootSignature(const class FD3D12RayTracingShader* Shader);
 	const FD3D12RootSignature* GetGlobalRayTracingRootSignature(const FRHIShaderBindingLayout& ShaderBindingLayout);
 
-#if USE_STATIC_ROOT_SIGNATURE
-	FORCEINLINE const FD3D12RootSignature* GetStaticGraphicsWithConstantsRootSignature() const
-	{
-		return &StaticGraphicsWithConstantsRootSignature;
-	}
-
-	FORCEINLINE const FD3D12RootSignature* GetStaticComputeWithConstantsRootSignature() const
-	{
-		return &StaticComputeWithConstantsRootSignature;
-	}
-#endif
-
 	FORCEINLINE FD3D12RootSignatureManager* GetRootSignatureManager()
 	{
 		return &RootSignatureManager;
diff --git a/Engine/Source/Runtime/D3D12RHI/Private/D3D12Util.cpp b/Engine/Source/Runtime/D3D12RHI/Private/D3D12Util.cpp
--- a/Engine/Source/Runtime/D3D12RHI/Private/D3D12Util.cpp
+++ b/Engine/Source/Runtime/D3D12RHI/Private/D3D12Util.cpp
@@ -1467,6 +1467,38 @@ const FD3D12RootSignature* FD3D12Adapter::GetRootSignature(const FD3D12ComputeSh
 #endif //! USE_STATIC_ROOT_SIGNATURE
 }
 
+const FD3D12RootSignature* FD3D12Adapter::GetRootSignature(ERHIShaderBundleMode ShaderBundleMode, bool bBindless)
+{
+	// Determine the root signature for the bundle
+#if USE_STATIC_ROOT_SIGNATURE
+	if (ShaderBundleMode == ERHIShaderBundleMode::CS)
+	{
+#if PLATFORM_SUPPORTS_BINDLESS_RENDERING
+		if (bBindless)
+		{
+			return &StaticBindlessComputeWithConstantsRootSignature;
+		}
+#endif
+		return &StaticComputeWithConstantsRootSignature;
+	}
+
+	if (ShaderBundleMode == ERHIShaderBundleMode::MSPS || ShaderBundleMode == ERHIShaderBundleMode::VSPS)
+	{
+#if PLATFORM_SUPPORTS_BINDLESS_RENDERING
+		if (bBindless)
+		{
+			return &StaticBindlessGraphicsWithConstantsRootSignature;
+		}
+#endif
+
+		return &StaticGraphicsWithConstantsRootSignature;
+	}
+#endif // USE_STATIC_ROOT_SIGNATURE
+
+	checkNoEntry();
+	return nullptr;
+}
+
 const FD3D12RootSignature* FD3D12Adapter::GetRootSignature(const FD3D12WorkGraphShader* WorkGraphShader)
 {
 	check(WorkGraphShader);
