D3D12 Bindless - only request a new active GPU heap if one was already active

diff --git a/Engine/Source/Runtime/D3D12RHI/Private/D3D12BindlessDescriptors.cpp b/Engine/Source/Runtime/D3D12RHI/Private/D3D12BindlessDescriptors.cpp
--- a/Engine/Source/Runtime/D3D12RHI/Private/D3D12BindlessDescriptors.cpp
+++ b/Engine/Source/Runtime/D3D12RHI/Private/D3D12BindlessDescriptors.cpp
@@ -287,8 +287,12 @@ void FD3D12BindlessResourceManager::GrowCPUHeap(uint32 OriginalNumDescriptors, u
 	UE::D3D12Descriptors::CopyDescriptors(GetParentDevice(), NewCpuHeap, CpuHeap, 0, OriginalNumDescriptors);
 	CpuHeap = NewCpuHeap;
 
-	bRequestNewActiveGpuHeap = true;
-	bCPUHeapResized = true;
+	// Only request a new active GPU heap if one was already active.
+	if (ActiveGpuHeapIndex >= 0)
+	{
+		bRequestNewActiveGpuHeap = true;
+		bCPUHeapResized = true;
+	}
 }
 
 void FD3D12BindlessResourceManager::CleanupResources()
