Plumb bDepthBounds on XB1

#FailedValidations Preflight

diff --git a/Engine/Source/Runtime/D3D12RHI/Private/D3D12State.cpp b/Engine/Source/Runtime/D3D12RHI/Private/D3D12State.cpp
--- a/Engine/Source/Runtime/D3D12RHI/Private/D3D12State.cpp
+++ b/Engine/Source/Runtime/D3D12RHI/Private/D3D12State.cpp
@@ -404,10 +404,6 @@ FDepthStencilStateRHIRef FD3D12DynamicRHI::RHICreateDepthStencilState(const FDep
 	{
 		DepthStencilDesc.BackFace = DepthStencilDesc.FrontFace;
 	}
-#if PLATFORM_WINDOWS
-	// Currently, the initializer doesn't include depth bound test info, we have to track it separately.
-	DepthStencilDesc.DepthBoundsTestEnable = false;
-#endif
 
 	const bool bStencilOpIsKeep =
 		Initializer.FrontFaceStencilFailStencilOp == SO_Keep
