Add Option to Skip Discarding Depth

On some platforms it's a win to skip discarding depth during aliasing operations. This adds an option to turn that behavior on in both the legacy and enhanced barriers implementations.

diff --git a/Engine/Source/Runtime/D3D12RHI/Private/D3D12EnhancedBarriers.cpp b/Engine/Source/Runtime/D3D12RHI/Private/D3D12EnhancedBarriers.cpp
--- a/Engine/Source/Runtime/D3D12RHI/Private/D3D12EnhancedBarriers.cpp
+++ b/Engine/Source/Runtime/D3D12RHI/Private/D3D12EnhancedBarriers.cpp
@@ -2020,8 +2020,6 @@ void FD3D12EnhancedBarriersBatcher::AddTextureBarrier(
 		LayoutBefore,
 		LayoutAfter))
 	{
-		bDiscard = bDiscard && GD3D12AllowDiscardResources;
-
 		const D3D12_TEXTURE_BARRIER_FLAGS Flags = 
 			bDiscard ?
 				D3D12_TEXTURE_BARRIER_FLAG_DISCARD :
@@ -3151,7 +3149,9 @@ bool FD3D12EnhancedBarriersForContext::AddBarriersForTransitionData(
 
 				const bool bDiscard =
 					LayoutBefore == D3D12_BARRIER_LAYOUT_UNDEFINED
-					&& LayoutAfter != D3D12_BARRIER_LAYOUT_UNDEFINED;
+					&& LayoutAfter != D3D12_BARRIER_LAYOUT_UNDEFINED
+					&& !(GD3D12DisableDiscardOfDepthResources && Resource->IsDepthStencilResource())
+					&& GD3D12AllowDiscardResources;
 
 				Batcher->AddTextureBarrier(
 					Context,
diff --git a/Engine/Source/Runtime/D3D12RHI/Private/D3D12LegacyBarriers.cpp b/Engine/Source/Runtime/D3D12RHI/Private/D3D12LegacyBarriers.cpp
--- a/Engine/Source/Runtime/D3D12RHI/Private/D3D12LegacyBarriers.cpp
+++ b/Engine/Source/Runtime/D3D12RHI/Private/D3D12LegacyBarriers.cpp
@@ -1280,6 +1280,11 @@ void FD3D12LegacyBarriersForContext::HandleDiscardResources(
 		else
 #endif // #if PLATFORM_REQUIRES_TYPELESS_RESOURCE_DISCARD_WORKAROUND
 		{
+			if (GD3D12DisableDiscardOfDepthResources && DiscardResource.Resource->IsDepthStencilResource())
+			{
+				continue;
+			}
+
 			if (DiscardResource.Subresource == D3D12_RESOURCE_BARRIER_ALL_SUBRESOURCES)
 			{
 				Context.GetCommandList().GraphicsCommandList()->DiscardResource(DiscardResource.Resource->GetResource(), nullptr);
diff --git a/Engine/Source/Runtime/D3D12RHI/Private/ID3D12Barriers.h b/Engine/Source/Runtime/D3D12RHI/Private/ID3D12Barriers.h
--- a/Engine/Source/Runtime/D3D12RHI/Private/ID3D12Barriers.h
+++ b/Engine/Source/Runtime/D3D12RHI/Private/ID3D12Barriers.h
@@ -33,6 +33,12 @@ inline FAutoConsoleVariableRef CVarD3D12AllowDiscardResources(
 	TEXT("Whether to call DiscardResources after transient aliasing acquire. This is not needed on some platforms if newly acquired resources are cleared before use."),
 	ECVF_RenderThreadSafe);
 
+inline int32 GD3D12DisableDiscardOfDepthResources = 0;
+inline FAutoConsoleVariableRef CVarDisableDiscardOfDepthResources(
+	TEXT("d3d12.DisableDiscardOfDepthResources"),
+	GD3D12DisableDiscardOfDepthResources,
+	TEXT("Whether to skip discarding depth resources after transient aliasing acquire. This is not needed on some platforms if the whole (sub)resource is written before it's read."));
+
 inline int32 GD3D12BatchResourceBarriers = 1;
 inline FAutoConsoleVariableRef CVarD3D12BatchResourceBarriers(
 	TEXT("d3d12.BatchResourceBarriers"),
