Flush the GPU before creating a SwapChain. This is an attempt to prevent issues from HWNDs being reused by the Windows runtime.

diff --git a/Engine/Source/Runtime/D3D12RHI/Private/Windows/WindowsD3D12Viewport.cpp b/Engine/Source/Runtime/D3D12RHI/Private/Windows/WindowsD3D12Viewport.cpp
--- a/Engine/Source/Runtime/D3D12RHI/Private/Windows/WindowsD3D12Viewport.cpp
+++ b/Engine/Source/Runtime/D3D12RHI/Private/Windows/WindowsD3D12Viewport.cpp
@@ -43,6 +43,11 @@ FD3D12Viewport::FD3D12Viewport(class FD3D12Adapter* InParent, HWND InWindowHandl
 void FD3D12Viewport::Init()
 {
 	FD3D12Adapter* Adapter = GetParentAdapter();
+
+	// Make sure we flush any pending GPU work that uses existing swapchains before we start using a potentially recycled HWND.
+	FlushRenderingCommands();
+	Adapter->BlockUntilIdle();
+
 	IDXGIFactory2* Factory2 = Adapter->GetDXGIFactory2();
 
 	TArray<IDXGISwapchainProvider*> DXGISwapchainProviderModules = IModularFeatures::Get().GetModularFeatureImplementations<IDXGISwapchainProvider>(IDXGISwapchainProvider::GetModularFeatureName());
