- Always check for WaveOpsSupport in D3D12Adapter, remove feature level restriction
- Get the Per Platform value of r.TemporalAA.Mobile.UseCompute

diff --git a/Engine/Source/Runtime/D3D12RHI/Private/D3D12Adapter.cpp b/Engine/Source/Runtime/D3D12RHI/Private/D3D12Adapter.cpp
--- a/Engine/Source/Runtime/D3D12RHI/Private/D3D12Adapter.cpp
+++ b/Engine/Source/Runtime/D3D12RHI/Private/D3D12Adapter.cpp
@@ -1140,14 +1140,11 @@ void FD3D12Adapter::InitializeDevices()
 
 			// Detect availability of shader model 6.0 wave operations
 			{
-				if (GMaxRHIFeatureLevel >= ERHIFeatureLevel::SM6)
-				{
-					D3D12_FEATURE_DATA_D3D12_OPTIONS1 Features{};
-					RootDevice->CheckFeatureSupport(D3D12_FEATURE_D3D12_OPTIONS1, &Features, sizeof(Features));
-					GRHISupportsWaveOperations = Features.WaveOps;
-					GRHIMinimumWaveSize = Features.WaveLaneCountMin;
-					GRHIMaximumWaveSize = Features.WaveLaneCountMax;
-				}
+				D3D12_FEATURE_DATA_D3D12_OPTIONS1 Features{};
+				RootDevice->CheckFeatureSupport(D3D12_FEATURE_D3D12_OPTIONS1, &Features, sizeof(Features));
+				GRHISupportsWaveOperations = Features.WaveOps;
+				GRHIMinimumWaveSize = Features.WaveLaneCountMin;
+				GRHIMaximumWaveSize = Features.WaveLaneCountMax;
 
 				if (GRHISupportsWaveOperations)
 				{
diff --git a/Engine/Source/Runtime/Renderer/Private/PostProcess/TemporalAA.cpp b/Engine/Source/Runtime/Renderer/Private/PostProcess/TemporalAA.cpp
--- a/Engine/Source/Runtime/Renderer/Private/PostProcess/TemporalAA.cpp
+++ b/Engine/Source/Runtime/Renderer/Private/PostProcess/TemporalAA.cpp
@@ -531,7 +531,8 @@ float GetTemporalAAHistoryUpscaleFactor(const FViewInfo& View)
 
 bool DoesTemporalAAUseComputeShader(EShaderPlatform Platform)
 {
-	return !IsMobilePlatform(Platform) || CVarTemporalAAMobileUseCompute.GetValueOnAnyThread() > 0;
+	static FShaderPlatformCachedIniValue<int32> TemporalAAMobileUseCompute(TEXT("r.TemporalAA.Mobile.UseCompute"));
+	return !IsMobilePlatform(Platform) || (TemporalAAMobileUseCompute.Get(Platform) > 0);
 }
 
 FIntPoint FTAAPassParameters::GetOutputExtent() const
