D3D12, added more information to checked texture region copies

#FailedValidations Preflight

diff --git a/Engine/Source/Runtime/D3D12RHI/Private/D3D12Texture.cpp b/Engine/Source/Runtime/D3D12RHI/Private/D3D12Texture.cpp
--- a/Engine/Source/Runtime/D3D12RHI/Private/D3D12Texture.cpp
+++ b/Engine/Source/Runtime/D3D12RHI/Private/D3D12Texture.cpp
@@ -1453,6 +1453,17 @@ static uint64 GetPlacedBufferRequiredSize(FD3D12CommandContext& Context, const D
 	return SubresourceSizeAligned + SubresourceSizeUnaligned;
 }
 
+FString GetTextureCopyRegionString(const FSubresourceIndices& Indices, const D3D12_BOX& CopyBox, uint32 SubresourceWidth, uint32 SubresourceHeight, uint32 SubresourceDepth)
+{
+	return FString::Printf(
+		TEXT("Subresource (Mip: %u, Slice: %u, Plane: %u) with dimensions [%u, %u, %u], copying from [%u, %u, %u] to [%u, %u, %u]"),
+		Indices.MipIndex, Indices.SliceIndex, Indices.PlaneIndex,
+		SubresourceWidth, SubresourceHeight, SubresourceDepth,
+		CopyBox.left, CopyBox.top, CopyBox.front,
+		CopyBox.right, CopyBox.bottom, CopyBox.back
+	);
+}
+
 FORCENOINLINE void FD3D12CommandContext::CopyTextureRegionChecked(
 	const D3D12_TEXTURE_COPY_LOCATION* DestCopyLocation, int DestX, int DestY, int DestZ, EPixelFormat DestPixelFormat,
 	const D3D12_TEXTURE_COPY_LOCATION* SourceCopyLocation, const D3D12_BOX* SourceBox, EPixelFormat SourcePixelFormat,
@@ -1466,6 +1477,8 @@ FORCENOINLINE void FD3D12CommandContext::CopyTextureRegionChecked(
 	const FPixelFormatInfo& DestPixelFormatInfo = GPixelFormats[DestPixelFormat];
 	const FPixelFormatInfo& SourcePixelFormatInfo = GPixelFormats[SourcePixelFormat];
 
+#define LAZY_COPY_DECORATOR *GetTextureCopyRegionString(Indices, CopyBox, SubresourceWidth, SubresourceHeight, SubresourceDepth)
+
 	switch (SourceCopyLocation->Type)
 	{
 		// Copy from texture
@@ -1475,16 +1488,16 @@ FORCENOINLINE void FD3D12CommandContext::CopyTextureRegionChecked(
 
 			const D3D12_BOX CopyBox = SourceBox ? *SourceBox : GetSubresourceCopyBox(SourceDesc, DestCopyLocation, Indices, SourcePixelFormatInfo);
 
-			RHI_BREADCRUMB_CHECK_SHIPPINGF(*this, (CopyBox.left % SourcePixelFormatInfo.BlockSizeX) == 0 && (CopyBox.right % SourcePixelFormatInfo.BlockSizeX) == 0, TEXT("Width not aligned to block size for: '%s'"), *DebugName.ToString());
-			RHI_BREADCRUMB_CHECK_SHIPPINGF(*this, (CopyBox.top % SourcePixelFormatInfo.BlockSizeY) == 0 && (CopyBox.bottom % SourcePixelFormatInfo.BlockSizeY) == 0, TEXT("Height not aligned to block size for: '%s'"), *DebugName.ToString());
-				
 			const uint32 SubresourceWidth  = AlignArbitrary<uint32>(FMath::Max<uint32>(1u, SourceDesc.Width >> Indices.MipIndex), SourcePixelFormatInfo.BlockSizeX);
 			const uint32 SubresourceHeight = AlignArbitrary<uint32>(FMath::Max<uint32>(1u, SourceDesc.Height >> Indices.MipIndex), SourcePixelFormatInfo.BlockSizeY);
 			const uint32 SubresourceDepth  = AlignArbitrary<uint32>(SourceDesc.Dimension == D3D12_RESOURCE_DIMENSION_TEXTURE3D ? FMath::Max<uint32>(1u, SourceDesc.DepthOrArraySize >> Indices.MipIndex) : 1, SourcePixelFormatInfo.BlockSizeZ);
 
-			RHI_BREADCRUMB_CHECK_SHIPPINGF(*this, CopyBox.left <= CopyBox.right && CopyBox.right <= SubresourceWidth, TEXT("Source width out of bounds for: '%s'"), *DebugName.ToString());
-			RHI_BREADCRUMB_CHECK_SHIPPINGF(*this, CopyBox.top <= CopyBox.bottom && CopyBox.bottom <= SubresourceHeight, TEXT("Source height out of bounds for: '%s'"), *DebugName.ToString());
-			RHI_BREADCRUMB_CHECK_SHIPPINGF(*this, CopyBox.front <= CopyBox.back && CopyBox.back <= SubresourceDepth, TEXT("Source depth out of bounds for: '%s'"), *DebugName.ToString());
+			RHI_BREADCRUMB_CHECK_SHIPPINGF(*this, (CopyBox.left % SourcePixelFormatInfo.BlockSizeX) == 0 && (CopyBox.right % SourcePixelFormatInfo.BlockSizeX) == 0, TEXT("Width not aligned to block size for: '%s', %s"), *DebugName.ToString(), LAZY_COPY_DECORATOR);
+			RHI_BREADCRUMB_CHECK_SHIPPINGF(*this, (CopyBox.top % SourcePixelFormatInfo.BlockSizeY) == 0 && (CopyBox.bottom % SourcePixelFormatInfo.BlockSizeY) == 0, TEXT("Height not aligned to block size for: '%s', %s"), *DebugName.ToString(), LAZY_COPY_DECORATOR);
+				
+			RHI_BREADCRUMB_CHECK_SHIPPINGF(*this, CopyBox.left <= CopyBox.right && CopyBox.right <= SubresourceWidth, TEXT("Source width out of bounds for: '%s', %s"), *DebugName.ToString(), LAZY_COPY_DECORATOR);
+			RHI_BREADCRUMB_CHECK_SHIPPINGF(*this, CopyBox.top <= CopyBox.bottom && CopyBox.bottom <= SubresourceHeight, TEXT("Source height out of bounds for: '%s', %s"), *DebugName.ToString(), LAZY_COPY_DECORATOR);
+			RHI_BREADCRUMB_CHECK_SHIPPINGF(*this, CopyBox.front <= CopyBox.back && CopyBox.back <= SubresourceDepth, TEXT("Source depth out of bounds for: '%s', %s"), *DebugName.ToString(), LAZY_COPY_DECORATOR);
 			break;
 		}
 
@@ -1505,16 +1518,16 @@ FORCENOINLINE void FD3D12CommandContext::CopyTextureRegionChecked(
 			const FSubresourceIndices Indices = GetSubresourceIndices(DestDesc, DestCopyLocation->SubresourceIndex);
 			const D3D12_BOX CopyBox = SourceBox ? GetBoxAdjustedForPixelFormat(*SourceBox, SourcePixelFormatInfo, DestPixelFormatInfo) : GetSubresourceCopyBox(DestDesc, SourceCopyLocation, Indices, DestPixelFormatInfo);
 				
-			RHI_BREADCRUMB_CHECK_SHIPPINGF(*this, DestX % DestPixelFormatInfo.BlockSizeX == 0 && (CopyBox.right - CopyBox.left) % DestPixelFormatInfo.BlockSizeX == 0, TEXT("Width not aligned to block size for: '%s'"), *DebugName.ToString());
-			RHI_BREADCRUMB_CHECK_SHIPPINGF(*this, DestY % DestPixelFormatInfo.BlockSizeY == 0 && (CopyBox.bottom - CopyBox.top) % DestPixelFormatInfo.BlockSizeY == 0, TEXT("Height not aligned to block size for: '%s'"), *DebugName.ToString());
-				
 			const uint32 SubresourceWidth  = AlignArbitrary<uint32>(FMath::Max<uint32>(1u, DestDesc.Width >> Indices.MipIndex), DestPixelFormatInfo.BlockSizeX);
 			const uint32 SubresourceHeight = AlignArbitrary<uint32>(FMath::Max<uint32>(1u, DestDesc.Height >> Indices.MipIndex), DestPixelFormatInfo.BlockSizeY);
 			const uint32 SubresourceDepth  = AlignArbitrary<uint32>(DestDesc.Dimension == D3D12_RESOURCE_DIMENSION_TEXTURE3D ? FMath::Max<uint32>(1u, DestDesc.DepthOrArraySize >> Indices.MipIndex) : 1, DestPixelFormatInfo.BlockSizeZ);
 
-			RHI_BREADCRUMB_CHECK_SHIPPINGF(*this, DestX + (CopyBox.right - CopyBox.left) <= SubresourceWidth, TEXT("Dest width out of bounds for: '%s'"), *DebugName.ToString());
-			RHI_BREADCRUMB_CHECK_SHIPPINGF(*this, DestY + (CopyBox.bottom - CopyBox.top) <= SubresourceHeight, TEXT("Dest height out of bounds for: '%s'"), *DebugName.ToString());
-			RHI_BREADCRUMB_CHECK_SHIPPINGF(*this, DestZ + (CopyBox.back - CopyBox.front) <= SubresourceDepth, TEXT("Dest depth out of bounds for: '%s'"), *DebugName.ToString());
+			RHI_BREADCRUMB_CHECK_SHIPPINGF(*this, DestX % DestPixelFormatInfo.BlockSizeX == 0 && (CopyBox.right - CopyBox.left) % DestPixelFormatInfo.BlockSizeX == 0, TEXT("Width not aligned to block size for: '%s', %s"), *DebugName.ToString(), LAZY_COPY_DECORATOR);
+			RHI_BREADCRUMB_CHECK_SHIPPINGF(*this, DestY % DestPixelFormatInfo.BlockSizeY == 0 && (CopyBox.bottom - CopyBox.top) % DestPixelFormatInfo.BlockSizeY == 0, TEXT("Height not aligned to block size for: '%s', %s"), *DebugName.ToString(), LAZY_COPY_DECORATOR);
+				
+			RHI_BREADCRUMB_CHECK_SHIPPINGF(*this, DestX + (CopyBox.right - CopyBox.left) <= SubresourceWidth, TEXT("Dest width out of bounds for: '%s', %s"), *DebugName.ToString(), LAZY_COPY_DECORATOR);
+			RHI_BREADCRUMB_CHECK_SHIPPINGF(*this, DestY + (CopyBox.bottom - CopyBox.top) <= SubresourceHeight, TEXT("Dest height out of bounds for: '%s', %s"), *DebugName.ToString(), LAZY_COPY_DECORATOR);
+			RHI_BREADCRUMB_CHECK_SHIPPINGF(*this, DestZ + (CopyBox.back - CopyBox.front) <= SubresourceDepth, TEXT("Dest depth out of bounds for: '%s', %s"), *DebugName.ToString(), LAZY_COPY_DECORATOR);
 			break;
 		}
 
@@ -1526,6 +1539,8 @@ FORCENOINLINE void FD3D12CommandContext::CopyTextureRegionChecked(
 			break;
 		}
 	}
+
+#undef LAZY_COPY_DECORATOR
 #endif // ENABLE_COPY_TEXTURE_REGION_CHECK
 
 	// Just pass down callchain
