Clamp d3d12 scissor rect top / left to zero. negative values leads to badness on certain d3d12 devices

#rnx

diff --git a/Engine/Source/Runtime/D3D12RHI/Private/D3D12Commands.cpp b/Engine/Source/Runtime/D3D12RHI/Private/D3D12Commands.cpp
--- a/Engine/Source/Runtime/D3D12RHI/Private/D3D12Commands.cpp
+++ b/Engine/Source/Runtime/D3D12RHI/Private/D3D12Commands.cpp
@@ -344,17 +344,29 @@ void FD3D12CommandContext::RHISetStereoViewport(float LeftMinX, float RightMinX,
 
 void FD3D12CommandContext::RHISetScissorRect(bool bEnable, uint32 MinX, uint32 MinY, uint32 MaxX, uint32 MaxY)
 {
+	CD3DX12_RECT ScissorRect;
+
 	if (bEnable)
 	{
-		const CD3DX12_RECT ScissorRect(MinX, MinY, MaxX, MaxY);
-		StateCache.SetScissorRect(ScissorRect);
+		ScissorRect.left = MinX;
+		ScissorRect.top = MinY;
+		ScissorRect.right = MaxX;
+		ScissorRect.bottom = MaxY;
 	}
 	else
-	{
+	{ 
 		const D3D12_VIEWPORT& Viewport = StateCache.GetViewport();
-		const CD3DX12_RECT ScissorRect((LONG) Viewport.TopLeftX, (LONG) Viewport.TopLeftY, (LONG) Viewport.TopLeftX + (LONG) Viewport.Width, (LONG) Viewport.TopLeftY + (LONG) Viewport.Height);
-		StateCache.SetScissorRect(ScissorRect);
+		ScissorRect.left = (LONG)Viewport.TopLeftX;
+		ScissorRect.top = (LONG)Viewport.TopLeftY;
+		ScissorRect.right = (LONG)Viewport.TopLeftX + (LONG)Viewport.Width;
+		ScissorRect.bottom = (LONG)Viewport.TopLeftY + (LONG)Viewport.Height;
 	}
+
+	// Ensure left / top are clamped to 0. Negative values leads to badness
+	ScissorRect.left = FMath::Max(0, ScissorRect.left);
+	ScissorRect.top = FMath::Max(0, ScissorRect.top);
+
+	StateCache.SetScissorRect(ScissorRect);
 }
 
 static void ApplyStaticUniformBuffersOnContext(FD3D12CommandContext& Context, FRHIShader* Shader)
