D3D12 Bindless Improvements
- Prevent allocating a zero-sized descriptor heap.
- Log descriptor heap sizes on init.
- When logging the failure to grow the descriptor heap, log out the size that failed.
- When clearing a UAV, don't allocate a bindless descriptor if bindless isn't enabled on the context.

diff --git a/Engine/Source/Runtime/D3D12RHI/Private/D3D12Adapter.cpp b/Engine/Source/Runtime/D3D12RHI/Private/D3D12Adapter.cpp
--- a/Engine/Source/Runtime/D3D12RHI/Private/D3D12Adapter.cpp
+++ b/Engine/Source/Runtime/D3D12RHI/Private/D3D12Adapter.cpp
@@ -1124,7 +1124,7 @@ void FD3D12Adapter::InitializeDevices()
 			if (GetHighestShaderModel() >= D3D_SHADER_MODEL_6_6 && GetResourceBindingTier() >= D3D12_RESOURCE_BINDING_TIER_3)
 			{
 				GRHIGlobals.bSupportsBindless = (GMaxRHIFeatureLevel > ERHIFeatureLevel::SM5);
-				UE_LOG(LogD3D12RHI, Log, TEXT("Bindless resources are supported"));
+				UE_LOG(LogD3D12RHI, Log, TEXT("Bindless is supported. MaxNonSamplerDescriptors: %d, MaxSamplerDescriptors: %d"), MaxNonSamplerDescriptors, MaxSamplerDescriptors);
 			}
 
 			{
diff --git a/Engine/Source/Runtime/D3D12RHI/Private/D3D12BindlessDescriptors.cpp b/Engine/Source/Runtime/D3D12RHI/Private/D3D12BindlessDescriptors.cpp
--- a/Engine/Source/Runtime/D3D12RHI/Private/D3D12BindlessDescriptors.cpp
+++ b/Engine/Source/Runtime/D3D12RHI/Private/D3D12BindlessDescriptors.cpp
@@ -127,17 +127,29 @@ void FD3D12BindlessDescriptorAllocator::Init()
 
 	if (BindlessConfiguration != ERHIBindlessConfiguration::Disabled)
 	{
+		uint32 NumResourceDescriptors = FMath::Max<int32>(GBindlessResourceDescriptorHeapSize, 1);
+#if D3D12RHI_USE_CONSTANT_BUFFER_VIEWS
+		NumResourceDescriptors += GBindlessOnlineDescriptorHeapBlockSize;
+#endif
+
+		uint32 NumSamplerDescriptors = FMath::Max<int32>(GBindlessSamplerDescriptorHeapSize, 1);
+		if (NumSamplerDescriptors > D3D12_MAX_SHADER_VISIBLE_SAMPLER_HEAP_SIZE)
+		{
+			UE_LOG(LogD3D12RHI, Error, TEXT("D3D12.Bindless.SamplerDescriptorHeapSize was set to %d, which is higher than the D3D12 maximum of %d. Adjusting the value to prevent a crash."),
+				NumSamplerDescriptors,
+				D3D12_MAX_SHADER_VISIBLE_SAMPLER_HEAP_SIZE
+			);
+			NumSamplerDescriptors = D3D12_MAX_SHADER_VISIBLE_SAMPLER_HEAP_SIZE;
+		}
+
+		UE_LOG(LogD3D12RHI, Log, TEXT("Bindless is enabled (Configuration=%s). Initializing allocators with %d Resource and %d Sampler descriptors"), GetBindlessConfigurationString(BindlessConfiguration), NumResourceDescriptors, NumSamplerDescriptors);
+
 		const TStatId ResourceStats[] =
 		{
 			GET_STATID(STAT_ResourceDescriptorsAllocated),
 			GET_STATID(STAT_BindlessResourceDescriptorsAllocated),
 		};
 
-		uint32 NumResourceDescriptors = GBindlessResourceDescriptorHeapSize;
-#if D3D12RHI_USE_CONSTANT_BUFFER_VIEWS
-		NumResourceDescriptors += GBindlessOnlineDescriptorHeapBlockSize;
-#endif
-
 		ResourceAllocator = new FRHIHeapDescriptorAllocator(D3D12DescriptorTypeMaskFromHeapType(ERHIDescriptorHeapType::Standard), NumResourceDescriptors, ResourceStats);
 
 		const TStatId SamplerStats[] =
@@ -146,16 +158,6 @@ void FD3D12BindlessDescriptorAllocator::Init()
 			GET_STATID(STAT_BindlessSamplerDescriptorsAllocated),
 		};
 
-		uint32 NumSamplerDescriptors = GBindlessSamplerDescriptorHeapSize;
-		if (NumSamplerDescriptors > D3D12_MAX_SHADER_VISIBLE_SAMPLER_HEAP_SIZE)
-		{
-			UE_LOG(LogD3D12RHI, Error, TEXT("D3D12.Bindless.SamplerDescriptorHeapSize was set to %d, which is higher than the D3D12 maximum of %d. Adjusting the value to prevent a crash."),
-				NumSamplerDescriptors,
-				D3D12_MAX_SHADER_VISIBLE_SAMPLER_HEAP_SIZE
-			);
-			NumSamplerDescriptors = D3D12_MAX_SHADER_VISIBLE_SAMPLER_HEAP_SIZE;
-		}
-
 		SamplerAllocator = new FRHIHeapDescriptorAllocator(D3D12DescriptorTypeMaskFromHeapType(ERHIDescriptorHeapType::Sampler), NumSamplerDescriptors, SamplerStats);
 	}
 }
@@ -226,7 +228,7 @@ TOptional<FRHIDescriptorAllocation> FD3D12BindlessDescriptorAllocator::ResizeGro
 
 	if (CurrentNumDescriptors == NewNumDescriptors)
 	{
-		UE_LOG(LogD3D12RHI, Fatal, TEXT("Hit D3D12 device limits on descriptors when attempting to allocate a larger descriptor heap."));
+		UE_LOG(LogD3D12RHI, Fatal, TEXT("Hit D3D12 device limits on descriptors when attempting to allocate larger descriptor heap of size %d."), NewNumDescriptors);
 	}
 
 	TOptional<FRHIDescriptorAllocation> Allocation = ResourceAllocator->ResizeGrowAndAllocate(NewNumDescriptors, NumAllocations);
diff --git a/Engine/Source/Runtime/D3D12RHI/Private/D3D12UAV.cpp b/Engine/Source/Runtime/D3D12RHI/Private/D3D12UAV.cpp
--- a/Engine/Source/Runtime/D3D12RHI/Private/D3D12UAV.cpp
+++ b/Engine/Source/Runtime/D3D12RHI/Private/D3D12UAV.cpp
@@ -366,16 +366,26 @@ void FD3D12CommandContext::ClearUAV(TRHICommandList_RecursiveHazardous<FD3D12Com
 					R32UAVDesc.Buffer.NumElements = UAVDesc.Buffer.NumElements * DwordsPerElement;
 				}
 
+				Context.OpenIfNotAlready();
+
+				// The DescriptorType argument specifies which descriptor heap to allocate the view's descriptor from.
+				// If we won't actually need the descriptor (ie bindless is disabled), we can signal to not allocate it in the first place by specifying an Invalid descriptor type.
+				ERHIDescriptorType DescriptorType = ERHIDescriptorType::Invalid;
+#if PLATFORM_SUPPORTS_BINDLESS_RENDERING
+				if (Context.StateCache.GetDescriptorCache()->IsUsingBindlessHeap())
+				{
+					DescriptorType = ERHIDescriptorType::BufferUAV;
+				}
+#endif
+
 				// Scoped view will free the offline CPU handle once we return
 				// Always single GPU object, so FirstLinkedObject is nullptr
-				FD3D12UnorderedAccessView UAV(ParentDevice, nullptr, ERHIDescriptorType::BufferUAV);
+				FD3D12UnorderedAccessView UAV(ParentDevice, nullptr, DescriptorType);
 				UAV.CreateView(UnorderedAccessView->GetResourceLocation(), R32UAVDesc, FD3D12UnorderedAccessView::EFlags::None);
 
 				FD3D12OfflineDescriptor OfflineHandle = UAV.GetOfflineCpuHandle();
 				D3D12_GPU_DESCRIPTOR_HANDLE GPUHandle{};
 
-				Context.OpenIfNotAlready();
-
 #if PLATFORM_SUPPORTS_BINDLESS_RENDERING
 				if (Context.StateCache.GetDescriptorCache()->IsUsingBindlessHeap())
 				{
diff --git a/Engine/Source/Runtime/RHI/Private/RHIStrings.cpp b/Engine/Source/Runtime/RHI/Private/RHIStrings.cpp
--- a/Engine/Source/Runtime/RHI/Private/RHIStrings.cpp
+++ b/Engine/Source/Runtime/RHI/Private/RHIStrings.cpp
@@ -773,3 +773,18 @@ bool GetBindlessConfigurationFromString(FStringView InString, ERHIBindlessConfig
 
 	return false;
 }
+
+const TCHAR* GetBindlessConfigurationString(ERHIBindlessConfiguration InBindlessConfiguration)
+{
+	switch (InBindlessConfiguration)
+	{
+	case ERHIBindlessConfiguration::Disabled:   return TEXT("Disabled");
+	case ERHIBindlessConfiguration::RayTracing: return TEXT("RayTracing");
+	case ERHIBindlessConfiguration::Minimal:    return TEXT("Minimal");
+	case ERHIBindlessConfiguration::All:        return TEXT("All");
+	default:
+		checkf(false, TEXT("Unknown Bindless Configuration: %d"), static_cast<uint32>(InBindlessConfiguration));
+		return TEXT("");
+	}
+}
+
diff --git a/Engine/Source/Runtime/RHI/Public/RHIStrings.h b/Engine/Source/Runtime/RHI/Public/RHIStrings.h
--- a/Engine/Source/Runtime/RHI/Public/RHIStrings.h
+++ b/Engine/Source/Runtime/RHI/Public/RHIStrings.h
@@ -34,6 +34,9 @@ RHI_API void GetFeatureLevelName(ERHIFeatureLevel::Type InFeatureLevel, FName& O
 /** Finds a corresponding ERHIBindlessConfiguration given a string, or returns false if one could not be found. */
 RHI_API bool GetBindlessConfigurationFromString(FStringView InString, ERHIBindlessConfiguration& OutBindlessConfiguration);
 
+/** Get the string for the given BindlessConfiguration or return an empty string if invalid. */
+RHI_API const TCHAR* GetBindlessConfigurationString(ERHIBindlessConfiguration InBindlessConfiguration);
+
 /** Stringifies ERHIFeatureLevel */
 RHI_API FString LexToString(ERHIFeatureLevel::Type Level);
 
