Add missing barrier flush before copying data to the inline SBT.

diff --git a/Engine/Source/Runtime/D3D12RHI/Private/D3D12RayTracing.cpp b/Engine/Source/Runtime/D3D12RHI/Private/D3D12RayTracing.cpp
--- a/Engine/Source/Runtime/D3D12RHI/Private/D3D12RayTracing.cpp
+++ b/Engine/Source/Runtime/D3D12RHI/Private/D3D12RayTracing.cpp
@@ -1747,7 +1747,8 @@ public:
 
 			FD3D12Resource* SourceResource = UploadResourceLocation.GetResource();
 			FD3D12Resource* DestResource = D3D12InlineBindingDataBuffer->ResourceLocation.GetResource();
-						
+
+			Context.FlushResourceBarriers();
 			Context.UpdateResidency(DestResource);
 			Context.CopyBufferRegionChecked(
 				DestResource->GetResource(), DestResource->GetName(),
