When resizing the Bindless descriptor heap, make sure another thread didn't already resize it while you were waiting on the lock.

diff --git a/Engine/Source/Runtime/D3D12RHI/Private/D3D12BindlessDescriptors.cpp b/Engine/Source/Runtime/D3D12RHI/Private/D3D12BindlessDescriptors.cpp
--- a/Engine/Source/Runtime/D3D12RHI/Private/D3D12BindlessDescriptors.cpp
+++ b/Engine/Source/Runtime/D3D12RHI/Private/D3D12BindlessDescriptors.cpp
@@ -222,6 +222,12 @@ TOptional<FRHIDescriptorAllocation> FD3D12BindlessDescriptorAllocator::ResizeGro
 
 	FScopeLock ScopeLock(&ResourceHeapsCS);
 
+	// Once we have the lock, make sure we don't need to grow needlessly; a different thread may have resized while we were waiting.
+	if (TOptional<FRHIDescriptorAllocation> Allocation = ResourceAllocator->Allocate(NumAllocations))
+	{
+		return Allocation;
+	}
+
 	// Grow the descriptor handle allocator
 	uint32 CurrentNumDescriptors = ResourceAllocator->GetCapacity();
 	uint32 NewNumDescriptors = FMath::Min<uint32>(CurrentNumDescriptors * 2, MaxResourceHeapSize);
