Merging (as edit) CL 49949415

Remove workaround for NVIDIA async compute driver bug when running with driver 571.23 or newer.
- Original bug caused missing barriers between D3D12 command lists submitted to the async compute queue, when submitted together in a single call to ECL.
- This has now been fixed, so we don't need to limit async compute submission to a single D3D12 command list anymore.

diff --git a/Engine/Source/Runtime/D3D12RHI/Private/Windows/WindowsD3D12Device.cpp b/Engine/Source/Runtime/D3D12RHI/Private/Windows/WindowsD3D12Device.cpp
--- a/Engine/Source/Runtime/D3D12RHI/Private/Windows/WindowsD3D12Device.cpp
+++ b/Engine/Source/Runtime/D3D12RHI/Private/Windows/WindowsD3D12Device.cpp
@@ -1646,23 +1646,17 @@ void FD3D12DynamicRHI::Init()
 		// each. That's not the case if we execute more than one list at a time, so break it up into individual
 		// executions to make it true.
 		//
-		// @TODO - Add driver version check and message text once a fix is released. Also raise message priority to warning.
+		// Update 2026-01-19: NVIDIA reports the bug was fixed in driver version 571.23 and newer.
 		//
-
-		static constexpr auto D3D12SubmissionExecuteBatchSizeAsyncCVarName = TEXT("r.D3D12.Submission.MaxExecuteBatchSize.Async");
-		const auto CVarD3D12SubmissionAsyncMaxExecuteBatchSize = IConsoleManager::Get().FindConsoleVariable(D3D12SubmissionExecuteBatchSizeAsyncCVarName);
-		if (CVarD3D12SubmissionAsyncMaxExecuteBatchSize)
+		const auto CVarD3D12SubmissionAsyncMaxExecuteBatchSize = IConsoleManager::Get().FindConsoleVariable(TEXT("r.D3D12.Submission.MaxExecuteBatchSize.Async"));
+		if (!ensure(CVarD3D12SubmissionAsyncMaxExecuteBatchSize))
 		{
-			CVarD3D12SubmissionAsyncMaxExecuteBatchSize->Set(1);
-			UE_LOG(LogD3D12RHI, Display,
-				TEXT("Batched command list execution is disabled for async queues due to known bugs in the current driver."));
+			UE_LOG(LogD3D12RHI, Error, TEXT("Couldn't find CVar r.D3D12.Submission.MaxExecuteBatchSize.Async to work around a known bug in certain drivers."));
 		}
-		else
+		else if (FDriverVersion(GRHIGlobals.GpuInfo.AdapterUserDriverVersion) < FDriverVersion(TEXT("571.23")))
 		{
-			// If we're here, someone moved our cheese and this work-around logic needs to be updated
-			UE_LOG(LogD3D12RHI, Error,
-				TEXT("Couldn't find CVar %s to work around a known bug in the current driver."), D3D12SubmissionExecuteBatchSizeAsyncCVarName);
-			checkNoEntry();
+			CVarD3D12SubmissionAsyncMaxExecuteBatchSize->Set(1);
+			UE_LOG(LogD3D12RHI, Warning, TEXT("Batched command list execution is disabled for async queues due to known bugs in the current driver."));
 		}
 	} // if NVIDIA
 
