- Use  D3D12RHI_PLATFORM_LOW_LEVEL_RES_TRACK instead of D3D12RHI_PLATFORM_HAS_UNIFIED_MEMORY for tracking resource memory in Stats
- Add GRHIDeviceIsUMASeparateVA, GRHICacheCoherentUMA and GRHIDeviceSupportsGPUUploadHeaps for more detailed information about the device memory topology and capabilities

diff --git a/Engine/Source/Runtime/D3D12RHI/Private/D3D12Adapter.cpp b/Engine/Source/Runtime/D3D12RHI/Private/D3D12Adapter.cpp
--- a/Engine/Source/Runtime/D3D12RHI/Private/D3D12Adapter.cpp
+++ b/Engine/Source/Runtime/D3D12RHI/Private/D3D12Adapter.cpp
@@ -391,6 +391,8 @@ FD3D12AdapterDesc::FD3D12AdapterDesc(const DXGI_ADAPTER_DESC& InDesc, int32 InAd
 	, ResourceHeapTier(DeviceInfo.ResourceHeapTier)
 	, MaxRHIFeatureLevel(DeviceInfo.MaxRHIFeatureLevel)
 	, bUMA(DeviceInfo.bUMA)
+	, bCacheCoherentUMA(DeviceInfo.bCacheCoherentUMA)
+	, bSupportsGPUUploadHeaps(DeviceInfo.bSupportsGPUUploadHeaps)
 	, bSupportsWaveOps(DeviceInfo.bSupportsWaveOps)
 	, bSupportsAtomic64(DeviceInfo.bSupportsAtomic64)
 {
diff --git a/Engine/Source/Runtime/D3D12RHI/Private/D3D12Adapter.h b/Engine/Source/Runtime/D3D12RHI/Private/D3D12Adapter.h
--- a/Engine/Source/Runtime/D3D12RHI/Private/D3D12Adapter.h
+++ b/Engine/Source/Runtime/D3D12RHI/Private/D3D12Adapter.h
@@ -53,6 +53,8 @@ struct FD3D12DeviceBasicInfo
 	bool                        bSupportsWaveOps;
 	bool                        bSupportsAtomic64;
 	bool						bUMA;
+	bool						bCacheCoherentUMA;
+	bool						bSupportsGPUUploadHeaps;
 	ERHIFeatureLevel::Type      MaxRHIFeatureLevel;
 };
 
@@ -91,6 +93,12 @@ struct FD3D12AdapterDesc
 	/** Whether the GPU has unified memory. This can be used to identify integrated GPU as well. */
 	bool bUMA = false;
 
+	/** Whether the UMA device has cache coherency between the GPU and CPU.  */
+	bool bCacheCoherentUMA = false;
+
+	/** Whether the device supports GPU Upload Heaps  */
+	bool bSupportsGPUUploadHeaps = false;
+
 	/** Whether SM6.0 wave ops are supported */
 	bool bSupportsWaveOps = false;
 
diff --git a/Engine/Source/Runtime/D3D12RHI/Private/D3D12Stats.cpp b/Engine/Source/Runtime/D3D12RHI/Private/D3D12Stats.cpp
--- a/Engine/Source/Runtime/D3D12RHI/Private/D3D12Stats.cpp
+++ b/Engine/Source/Runtime/D3D12RHI/Private/D3D12Stats.cpp
@@ -38,8 +38,8 @@ void D3D12BufferStats::UpdateBufferStats(FD3D12Buffer& Buffer, bool bAllocating)
 	INC_MEMORY_STAT_BY_FName(GetD3D12BufferStat(BufferDesc).GetName(), RequestedSize);
 	INC_MEMORY_STAT_BY(STAT_D3D12MemoryCurrentTotal, RequestedSize);
 
-    // With D3D12RHI_PLATFORM_HAS_UNIFIED_MEMORY=1, MemoryTrace_Alloc is called during resource allocation
-#if UE_MEMORY_TRACE_ENABLED  && (D3D12RHI_PLATFORM_HAS_UNIFIED_MEMORY == 0)
+    // With D3D12RHI_PLATFORM_LOW_LEVEL_RES_TRACK=1, MemoryTrace_Alloc is called during resource allocation
+#if UE_MEMORY_TRACE_ENABLED  && (D3D12RHI_PLATFORM_LOW_LEVEL_RES_TRACK == 0)
 	const D3D12_GPU_VIRTUAL_ADDRESS GPUAddress = Location.GetGPUVirtualAddress();
 
 	if (bAllocating)
diff --git a/Engine/Source/Runtime/D3D12RHI/Private/Windows/WindowsD3D12Device.cpp b/Engine/Source/Runtime/D3D12RHI/Private/Windows/WindowsD3D12Device.cpp
--- a/Engine/Source/Runtime/D3D12RHI/Private/Windows/WindowsD3D12Device.cpp
+++ b/Engine/Source/Runtime/D3D12RHI/Private/Windows/WindowsD3D12Device.cpp
@@ -477,16 +477,30 @@ inline bool GetSupportsAtomic64(IDXGIAdapter* Adapter, ID3D12Device* Device)
 	return (D3D12Caps9.AtomicInt64OnTypedResourceSupported || CheckDeviceForEmulatedAtomic64Support(Adapter, Device));
 }
 
-inline bool IsDeviceUMA(ID3D12Device* device)
+inline bool GetDeviceUMAInformation(ID3D12Device* device, bool& OutUMA, bool& OutCacheCoherentUMA)
 {
 	D3D12_FEATURE_DATA_ARCHITECTURE data = { 0 };
 	if (SUCCEEDED(device->CheckFeatureSupport(D3D12_FEATURE_ARCHITECTURE, &data, sizeof(D3D12_FEATURE_DATA_ARCHITECTURE))))
 	{
-		return (bool)(data.UMA);
+		OutUMA = (bool)(data.UMA);
+		OutCacheCoherentUMA = (bool)(data.CacheCoherentUMA);
+		return true;
 	}
 	return false;
 }
 
+inline bool GetSupportGPUUploadHeaps(ID3D12Device* device)
+{
+	D3D12_FEATURE_DATA_D3D12_OPTIONS16 options16 = {};
+	
+	if (SUCCEEDED(device->CheckFeatureSupport(D3D12_FEATURE_D3D12_OPTIONS16, &options16, sizeof(options16))))
+	{
+		return (bool)(options16.GPUUploadHeapSupported);
+	}
+
+	return false;
+}
+
 /**
  * Attempts to create a D3D12 device for the adapter using at minimum MinFeatureLevel.
  * If creation is successful, true is returned and the max supported feature level is set in OutMaxFeatureLevel.
@@ -509,7 +523,12 @@ static bool SafeTestD3D12CreateDevice(IDXGIAdapter* Adapter, D3D_FEATURE_LEVEL M
 
 			OutInfo.MaxRHIFeatureLevel = FindMaxRHIFeatureLevel(OutInfo.MaxFeatureLevel, OutInfo.MaxShaderModel, OutInfo.ResourceBindingTier, OutInfo.bSupportsWaveOps, OutInfo.bSupportsAtomic64);
 
-			OutInfo.bUMA = IsDeviceUMA(Device);
+			if (GetDeviceUMAInformation(Device, OutInfo.bUMA, OutInfo.bCacheCoherentUMA) == false)
+			{
+				OutInfo.bUMA = OutInfo.bCacheCoherentUMA = false;
+			}
+
+			OutInfo.bSupportsGPUUploadHeaps = GetSupportGPUUploadHeaps(Device);
 
 			Device->Release();
 			return true;
@@ -901,12 +920,14 @@ void FD3D12DynamicRHIModule::FindAdapter()
 				);
 
 				UE_LOG(LogD3D12RHI, Log,
-					TEXT("  Adapter has %uMB of dedicated video memory, %uMB of dedicated system memory, and %uMB of shared system memory, %d output[s], UMA:%s"),
+					TEXT("  Adapter has %uMB of dedicated video memory, %uMB of dedicated system memory, and %uMB of shared system memory, %d output[s], UMA:%s CacheCoherentUMA:%s GPU Upload Heaps:%s"),
 					(uint32)(AdapterDesc.DedicatedVideoMemory / (1024 * 1024)),
 					(uint32)(AdapterDesc.DedicatedSystemMemory / (1024 * 1024)),
 					(uint32)(AdapterDesc.SharedSystemMemory / (1024 * 1024)),
 					OutputCount,
-					DeviceInfo.bUMA ? TEXT("true") : TEXT("false")
+					DeviceInfo.bUMA ? TEXT("true") : TEXT("false"),
+					DeviceInfo.bCacheCoherentUMA ? TEXT("true") : TEXT("false"),
+					DeviceInfo.bSupportsGPUUploadHeaps ? TEXT("true") : TEXT("false")
 				);
 
 				const bool bIsWARP = (RHIConvertToGpuVendorId(AdapterDesc.VendorId) == EGpuVendorId::Microsoft);
@@ -1098,6 +1119,9 @@ FDynamicRHI* FD3D12DynamicRHIModule::CreateRHI(ERHIFeatureLevel::Type RequestedF
 		const bool bIsIntegrated = ChosenAdapters[0].Get()->GetDesc().bUMA;
 		FGenericCrashContext::SetEngineData(TEXT("RHI.IntegratedGPU"), bIsIntegrated ? TEXT("true") : TEXT("false"));
 		GRHIDeviceIsIntegrated = bIsIntegrated;
+		GRHIDeviceIsUMASeparateVA = ChosenAdapters[0].Get()->GetDesc().bUMA;
+		GRHIDeviceIsCacheCoherentUMA = ChosenAdapters[0].Get()->GetDesc().bCacheCoherentUMA;
+		GRHIDeviceSupportsGPUUploadHeaps = ChosenAdapters[0].Get()->GetDesc().bSupportsGPUUploadHeaps;
 		UE_LOG(LogD3D12RHI, Log, TEXT("Integrated GPU (iGPU): %s"), GRHIDeviceIsIntegrated ? TEXT("true") : TEXT("false"));
 	}
 
diff --git a/Engine/Source/Runtime/D3D12RHI/Public/D3D12RHI.h b/Engine/Source/Runtime/D3D12RHI/Public/D3D12RHI.h
--- a/Engine/Source/Runtime/D3D12RHI/Public/D3D12RHI.h
+++ b/Engine/Source/Runtime/D3D12RHI/Public/D3D12RHI.h
@@ -25,6 +25,7 @@
 
 #if PLATFORM_WINDOWS
 	#define D3D12RHI_PLATFORM_HAS_UNIFIED_MEMORY	0
+	#define D3D12RHI_PLATFORM_LOW_LEVEL_RES_TRACK	0
 	#define ENABLE_RESIDENCY_MANAGEMENT				1
 	#define PIPELINE_STATE_FILE_LOCATION			FPaths::ProjectSavedDir()
 	#define USE_PIX									WITH_PIX_EVENT_RUNTIME
diff --git a/Engine/Source/Runtime/RHI/Public/RHIGlobals.h b/Engine/Source/Runtime/RHI/Public/RHIGlobals.h
--- a/Engine/Source/Runtime/RHI/Public/RHIGlobals.h
+++ b/Engine/Source/Runtime/RHI/Public/RHIGlobals.h
@@ -496,6 +496,15 @@ struct FRHIGlobals
 	/** Whether or not the current GPU is integrated into to CPU */
 	bool DeviceIsIntegrated = false;
 
+	/** Whether or not the current device has a Unified Memory Architecture (UMA) - Memory Shared across GPU and CPU but separate Virtual Addresses*/
+	bool DeviceIsUMASeparateVA = false;
+
+	/** Whether or not the current device UMA is Cache Coherent between the GPU and the CPU */
+	bool DeviceIsCacheCoherentUMA = false;
+
+	/** Whether or not the current device supports GPU Uploads Heaps (ReBar) */
+	bool DeviceSupportsGPUUploadHeaps = false;
+
 	/**
 	* Specifies the minimum and maximum number of lanes in the SIMD wave that this GPU can support. I.e. 32 on NVIDIA, 64 on AMD.
 	* Valid values are in range [4..128] (as per SM 6.0 specification) or 0 if unknown.
@@ -862,6 +871,9 @@ extern RHI_API uint32 GGPUFrameTime;
 #define GRHIRayTracingInstanceDescriptorSize                   GRHIGlobals.RayTracing.InstanceDescriptorSize
 #define GRHISupportsWaveOperations                             GRHIGlobals.SupportsWaveOperations
 #define GRHIDeviceIsIntegrated                                 GRHIGlobals.DeviceIsIntegrated
+#define GRHIDeviceIsUMASeparateVA                              GRHIGlobals.DeviceIsUMASeparateVA
+#define GRHIDeviceIsCacheCoherentUMA                           GRHIGlobals.DeviceIsCacheCoherentUMA
+#define GRHIDeviceSupportsGPUUploadHeaps                       GRHIGlobals.DeviceSupportsGPUUploadHeaps
 #define GRHIMinimumWaveSize                                    GRHIGlobals.MinimumWaveSize
 #define GRHIMaximumWaveSize                                    GRHIGlobals.MaximumWaveSize
 #define GRHISupportsRHIThread                                  GRHIGlobals.SupportsRHIThread
diff --git a/Engine/Source/Runtime/Windows/D3D11RHI/Private/Windows/WindowsD3D11Device.cpp b/Engine/Source/Runtime/Windows/D3D11RHI/Private/Windows/WindowsD3D11Device.cpp
--- a/Engine/Source/Runtime/Windows/D3D11RHI/Private/Windows/WindowsD3D11Device.cpp
+++ b/Engine/Source/Runtime/Windows/D3D11RHI/Private/Windows/WindowsD3D11Device.cpp
@@ -928,6 +928,8 @@ void FD3D11DynamicRHIModule::FindAdapter()
 	GRHIDeviceId = ChosenAdapter.DXGIAdapterDesc.DeviceId;
 	GRHIDeviceRevision = ChosenAdapter.DXGIAdapterDesc.Revision;
 	GRHIDeviceIsIntegrated = ChosenAdapter.bUMA;
+	GRHIDeviceIsUMASeparateVA = ChosenAdapter.bUMA;
+
 	UE_LOG(LogD3D11RHI, Log, TEXT("Integrated GPU (iGPU): %s"), GRHIDeviceIsIntegrated ? TEXT("true") : TEXT("false"));
 }
 
