D3D12RHI Cleanup
- Removing the unused FD3D12BucketAllocator and FD3D12SegListAllocator
- Moving D3D12Allocation allocator defines to D3D12RHIDefinitions
- Adding D3D12RHI_SUPPORTS_ALIASED_RHI_RESOURCES to gate XR specific code for Texture aliasing

diff --git a/Engine/Source/Runtime/D3D12RHI/Private/D3D12Allocation.cpp b/Engine/Source/Runtime/D3D12RHI/Private/D3D12Allocation.cpp
--- a/Engine/Source/Runtime/D3D12RHI/Private/D3D12Allocation.cpp
+++ b/Engine/Source/Runtime/D3D12RHI/Private/D3D12Allocation.cpp
@@ -22,30 +22,6 @@ static FAutoConsoleVariableRef CVarD3D12AllowPoolAllocateIndirectArgBuffers(
 	TEXT("Allow indirect args to be pool allocated (otherwise they will be committed resources) (default: 0)"),
 	ECVF_ReadOnly);
 
-#if D3D12RHI_SEGREGATED_TEXTURE_ALLOC
-static int32 GD3D12ReadOnlyTextureAllocatorMinPoolSize = 4 * 1024 * 1024;
-static FAutoConsoleVariableRef CVarD3D12ReadOnlyTextureAllocatorMinPoolSize(
-	TEXT("d3d12.ReadOnlyTextureAllocator.MinPoolSize"),
-	GD3D12ReadOnlyTextureAllocatorMinPoolSize,
-	TEXT("Minimum allocation granularity (in bytes) of each size list"),
-	ECVF_ReadOnly);
-
-static int32 GD3D12ReadOnlyTextureAllocatorMinNumToPool = 8;
-static FAutoConsoleVariableRef CVarD3D12ReadOnlyTextureAllocatorMinNumToPool(
-	TEXT("d3d12.ReadOnlyTextureAllocator.MinNumToPool"),
-	GD3D12ReadOnlyTextureAllocatorMinNumToPool,
-	TEXT("Texture pool of each size list must be large enough to store this")
-	TEXT("many textures unless constrained by maximum allocation granularity"),
-	ECVF_ReadOnly);
-
-static int32 GD3D12ReadOnlyTextureAllocatorMaxPoolSize = 20 * 1024 * 1024;
-static FAutoConsoleVariableRef CVarD3D12ReadOnlyTextureAllocatorMaxPoolSize(
-	TEXT("d3d12.ReadOnlyTextureAllocator.MaxPoolSize"),
-	GD3D12ReadOnlyTextureAllocatorMaxPoolSize,
-	TEXT("Maximum allocation granularity (in bytes) of each size list"),
-	ECVF_ReadOnly);
-#endif
-
 #if USE_BUFFER_POOL_ALLOCATOR
 
 #if !defined(BUFFER_POOL_DEFRAG_MAX_COPY_SIZE_PER_FRAME)
@@ -162,16 +138,6 @@ static FAutoConsoleVariableRef CVarD3D12FastConstantAllocatorPageSize(
 	TEXT("Page size for the fast constant allocator"),
 	ECVF_ReadOnly);
 
-
-#if D3D12RHI_SEGLIST_ALLOC_TRACK_WASTAGE
-static int32 GD3D12SegListTrackLeaks = 0;
-static FAutoConsoleVariableRef CVarD3D12SegListTrackLeaks(
-	TEXT("d3d12.SegListTrackLeaks"),
-	GD3D12SegListTrackLeaks,
-	TEXT("1: Enable leak tracking in d3d12 seglist's"),
-	ECVF_ReadOnly);
-#endif
-
 static int32 GD3D12FastAllocatorMinPagesToRetain = 5;
 static FAutoConsoleVariableRef CVarD3D12FastAllocatorMinPagesToRetain(
 	TEXT("d3d12.FastAllocator.MinPagesToRetain"),
@@ -882,222 +848,6 @@ void FD3D12MultiBuddyAllocator::Reset()
 
 }
 
-//-----------------------------------------------------------------------------
-//	Bucket Allocator
-//-----------------------------------------------------------------------------
-FD3D12BucketAllocator::FD3D12BucketAllocator(FD3D12Device* ParentDevice,
-	FRHIGPUMask VisibleNodes,
-	const FD3D12ResourceInitConfig& InInitConfig,
-	const FString& Name,
-	uint64 InBlockRetentionFrameCount) :
-	FD3D12ResourceAllocator(ParentDevice, VisibleNodes, InInitConfig, Name, 32 * 1024 * 1024),
-	BlockRetentionFrameCount(InBlockRetentionFrameCount)
-{}
-
-bool FD3D12BucketAllocator::TryAllocate(uint32 SizeInBytes, uint32 Alignment, FD3D12ResourceLocation& ResourceLocation)
-{
-	FScopeLock Lock(&CS);
-
-	FD3D12Adapter* Adapter = GetParentDevice()->GetParentAdapter();
-
-	// Alignment needs to be a multiple of 16 bytes
-	Alignment = Align(FMath::Max(Alignment, 16u), 16u);
-
-	// Size cannot be smaller than the requested alignment
-	SizeInBytes = FMath::Max(SizeInBytes, Alignment);
-
-	uint32 Bucket = BucketFromSize(SizeInBytes, BucketShift);
-	check(Bucket < NumBuckets);
-
-	uint32 BlockSize = BlockSizeFromBufferSize(SizeInBytes, BucketShift);
-
-	// If some odd alignment is requested, make sure the block can fulfill it.
-	if (BlockSize % Alignment != 0)
-	{
-		const uint32 AlignedSizeInBytes = SizeInBytes + Alignment;
-		Bucket = BucketFromSize(AlignedSizeInBytes, BucketShift);
-		BlockSize = BlockSizeFromBufferSize(AlignedSizeInBytes, BucketShift);
-	}
-
-	FD3D12BlockAllocatorPrivateData& Block = ResourceLocation.GetBlockAllocatorPrivateData();
-
-	// See if a block is already available in the bucket
-	if (AvailableBlocks[Bucket].Dequeue(Block))
-	{
-		check(Block.ResourceHeap);
-	}
-	else
-	{
-		// No blocks of the requested size are available so make one
-		FD3D12Resource* Resource = nullptr;
-		void* BaseAddress = nullptr;
-
-		// Allocate a block
-		check(BlockSize >= SizeInBytes);
-
-		if (FAILED(Adapter->CreateBuffer(InitConfig.HeapType, GetGPUMask(), GetVisibilityMask(), SizeInBytes < MIN_HEAP_SIZE ? MIN_HEAP_SIZE : SizeInBytes, &Resource, TEXT("BucketAllocator"), InitConfig.ResourceFlags)))
-		{
-			return false;
-		}
-
-		// Track the resource so we know when to delete it
-		SubAllocatedResources.Add(Resource);
-
-		if (IsCPUAccessible(InitConfig.HeapType))
-		{
-			BaseAddress = Resource->Map();
-			check(BaseAddress);
-			check(BaseAddress == (uint8*)(((uint64)BaseAddress + Alignment - 1) & ~((uint64)Alignment - 1)));
-		}
-
-		// Init the block we will return
-		Block.BucketIndex = Bucket;
-		Block.Offset = 0;
-		Block.ResourceHeap = Resource;
-		Block.ResourceHeap->AddRef();
-
-		// Chop up the rest of the resource into reusable blocks
-		if (BlockSize < MIN_HEAP_SIZE)
-		{
-			// Create additional available blocks that can be sub-allocated from the same resource
-			for (uint32 Offset = BlockSize; Offset <= MIN_HEAP_SIZE - BlockSize; Offset += BlockSize)
-			{
-				FD3D12BlockAllocatorPrivateData NewBlock = {};
-				NewBlock.BucketIndex = Bucket;
-				NewBlock.Offset = Offset;
-				NewBlock.ResourceHeap = Resource;
-				NewBlock.ResourceHeap->AddRef();
-
-				// Add the bucket to the available list
-				AvailableBlocks[Bucket].Enqueue(NewBlock);
-			}
-		}
-	}
-
-	uint64 AlignedBlockOffset = Block.Offset;
-	if (AlignedBlockOffset % Alignment != 0)
-	{
-		AlignedBlockOffset = AlignArbitrary(AlignedBlockOffset, Alignment);
-	}
-
-	ResourceLocation.SetType(FD3D12ResourceLocation::ResourceLocationType::eSubAllocation);
-	ResourceLocation.SetAllocator((FD3D12BaseAllocatorType*)this);
-	ResourceLocation.SetResource(Block.ResourceHeap);
-	ResourceLocation.SetSize(SizeInBytes);
-	ResourceLocation.SetOffsetFromBaseOfResource(AlignedBlockOffset);
-	ResourceLocation.SetGPUVirtualAddress(Block.ResourceHeap->GetGPUVirtualAddress() + AlignedBlockOffset);
-
-	if (IsCPUAccessible(InitConfig.HeapType))
-	{
-		ResourceLocation.SetMappedBaseAddress((void*)((uint64)Block.ResourceHeap->GetResourceBaseAddress() + AlignedBlockOffset));
-	}
-
-	// Check that when the offset is aligned that it doesn't go passed the end of the block
-	check(ResourceLocation.GetOffsetFromBaseOfResource() - Block.Offset + SizeInBytes <= BlockSize);
-
-	return true;
-}
-
-void FD3D12BucketAllocator::Deallocate(FD3D12ResourceLocation& ResourceLocation)
-{
-	FScopeLock Lock(&CS);
-
-	FD3D12Adapter* Adapter = GetParentDevice()->GetParentAdapter();
-	FD3D12ManualFence& FrameFence = Adapter->GetFrameFence();
-
-	FD3D12BlockAllocatorPrivateData& Block = ResourceLocation.GetBlockAllocatorPrivateData();
-	Block.FrameFence = FrameFence.GetNextFenceToSignal();
-
-	ExpiredBlocks.Enqueue(Block);
-}
-
-void FD3D12BucketAllocator::Initialize()
-{
-
-}
-
-void FD3D12BucketAllocator::Destroy()
-{
-	ReleaseAllResources();
-}
-void FD3D12BucketAllocator::CleanUpAllocations(uint64 InFrameLag)
-{
-	FScopeLock Lock(&CS);
-
-	FD3D12Adapter* Adapter = GetParentDevice()->GetParentAdapter();
-	FD3D12ManualFence& FrameFence = Adapter->GetFrameFence();
-
-#if SUB_ALLOCATED_DEFAULT_ALLOCATIONS
-	const static uint32 MinCleanupBucket = FMath::Max<uint32>(0, BucketFromSize(MIN_HEAP_SIZE, BucketShift) - 4);
-#else
-	const static uint32 MinCleanupBucket = 0;
-#endif
-
-	// Start at bucket 8 since smaller buckets are sub-allocated resources
-	// and would be fragmented by deleting blocks
-	for (uint32 bucket = MinCleanupBucket; bucket < NumBuckets; bucket++)
-	{
-		FD3D12BlockAllocatorPrivateData BlockInQueue = {};
-		const uint32 RetentionCount = BlockRetentionFrameCount;
-
-		const auto& Functor = [&FrameFence, RetentionCount](const FD3D12BlockAllocatorPrivateData& Block) { return FrameFence.IsFenceComplete(Block.FrameFence + RetentionCount, /* bUpdateCachedFenceValue */ false); };
-		while (AvailableBlocks[bucket].Dequeue(BlockInQueue, Functor))
-		{
-			SAFE_RELEASE(BlockInQueue.ResourceHeap);
-		}
-	}
-
-	FD3D12BlockAllocatorPrivateData BlockInQueue = {};
-
-	const auto& Functor = [&FrameFence](const FD3D12BlockAllocatorPrivateData& Block) { return FrameFence.IsFenceComplete(Block.FrameFence, /* bUpdateCachedFenceValue */ false); };
-	while (ExpiredBlocks.Dequeue(BlockInQueue, Functor))
-	{
-		// Add the bucket to the available list
-		AvailableBlocks[BlockInQueue.BucketIndex].Enqueue(BlockInQueue);
-	}
-}
-
-void FD3D12BucketAllocator::DumpAllocatorStats(class FOutputDevice& Ar)
-{
-	//TODO:
-}
-void FD3D12BucketAllocator::ReleaseAllResources()
-{
-	const static uint32 MinCleanupBucket = 0;
-
-	// Start at bucket 8 since smaller buckets are sub-allocated resources
-	// and would be fragmented by deleting blocks
-	for (uint32 bucket = MinCleanupBucket; bucket < NumBuckets; bucket++)
-	{
-		FD3D12BlockAllocatorPrivateData Block = {};
-		while (AvailableBlocks[bucket].Dequeue(Block))
-		{
-			SAFE_RELEASE(Block.ResourceHeap);
-		}
-	}
-
-	FD3D12BlockAllocatorPrivateData Block = {};
-
-	while (ExpiredBlocks.Dequeue(Block))
-	{
-		if (Block.BucketIndex >= MinCleanupBucket) //-V547
-		{
-			SAFE_RELEASE(Block.ResourceHeap);
-		}
-	}
-
-	for (FD3D12Resource*& Resource : SubAllocatedResources)
-	{
-		Resource->Release();
-		delete(Resource);
-	}
-}
-
-void FD3D12BucketAllocator::Reset()
-{
-
-}
-
 //-----------------------------------------------------------------------------
 //	Dynamic Buffer Allocator
 //-----------------------------------------------------------------------------
@@ -1978,79 +1728,8 @@ bool FD3D12TextureAllocatorPool::GetMemoryStats(uint64& OutTotalAllocated, uint6
 	return true;
 }
 
-#elif D3D12RHI_SEGREGATED_TEXTURE_ALLOC
-FD3D12TextureAllocatorPool::FD3D12TextureAllocatorPool(FD3D12Device* Device, FRHIGPUMask VisibilityNode) :
-	FD3D12DeviceChild(Device),
-	FD3D12MultiNodeGPUObject(Device->GetGPUMask(), VisibilityNode),
-	ReadOnlyTexturePool(
-		Device,
-		VisibilityNode,
-		D3D12_HEAP_TYPE_DEFAULT,
-		D3D12_HEAP_FLAG_ALLOW_ONLY_NON_RT_DS_TEXTURES,
-		GD3D12ReadOnlyTextureAllocatorMinPoolSize,
-		GD3D12ReadOnlyTextureAllocatorMinNumToPool,
-		GD3D12ReadOnlyTextureAllocatorMaxPoolSize)
-{
-}
-
-HRESULT FD3D12TextureAllocatorPool::AllocateTexture(
-	FD3D12ResourceDesc Desc,
-	const D3D12_CLEAR_VALUE* ClearValue,
-	EPixelFormat UEFormat,
-	FD3D12ResourceLocation& TextureLocation,
-	const D3D12_RESOURCE_STATES InitialState,
-	const TCHAR* Name)
-{
-	HRESULT RetCode = S_OK;
-	FD3D12Resource* NewResource = nullptr;
-	FD3D12Device* Device = GetParentDevice();
-	FD3D12Adapter* Adapter = Device->GetParentAdapter();
-
-	TextureLocation.Clear();
-
-	if (!EnumHasAnyFlags(Desc.Flags, D3D12_RESOURCE_FLAG_ALLOW_RENDER_TARGET|D3D12_RESOURCE_FLAG_ALLOW_DEPTH_STENCIL|D3D12_RESOURCE_FLAG_ALLOW_UNORDERED_ACCESS) &&
-		!Desc.NeedsUAVAliasWorkarounds() &&
-		//  4K align with NV12 causes a crash on HoloLens 2.
-		Desc.Format != DXGI_FORMAT_NV12 &&
-		Desc.SampleDesc.Count == 1)
-	{
-		// The top mip level must be less than 64 KB to use 4 KB alignment
-		Desc.Alignment = FD3D12Texture::CanBe4KAligned(Desc, UEFormat) ?
-			D3D12_SMALL_RESOURCE_PLACEMENT_ALIGNMENT :
-			D3D12_DEFAULT_RESOURCE_PLACEMENT_ALIGNMENT;
-		const D3D12_RESOURCE_ALLOCATION_INFO Info = Device->GetResourceAllocationInfoUncached(Desc);
-
-		TRefCountPtr<FD3D12SegHeap> BackingHeap;
-		const uint32 Offset = ReadOnlyTexturePool.Allocate(Info.SizeInBytes, Info.Alignment, BackingHeap);
-
-		if (Offset != FD3D12SegListAllocator::InvalidOffset)
-		{
-			RetCode = Adapter->CreatePlacedResource(Desc, BackingHeap.GetReference(), Offset, InitialState, ClearValue, &NewResource, Name, false);
-			if (SUCCEEDED(RetCode))
-			{
-				FD3D12SegListAllocatorPrivateData& PrivateData = TextureLocation.GetSegListAllocatorPrivateData();
-				PrivateData.Offset = Offset;
-
-				TextureLocation.SetType(FD3D12ResourceLocation::ResourceLocationType::eSubAllocation);
-				TextureLocation.SetSegListAllocator(&ReadOnlyTexturePool);
-				TextureLocation.SetSize(Info.SizeInBytes);
-				TextureLocation.SetOffsetFromBaseOfResource(Offset);
-				TextureLocation.SetResource(NewResource);
-
-				INC_DWORD_STAT(STAT_D3D12TextureAllocatorCount);
-			}
-			return RetCode;
-		}
-	}
-
-	const D3D12_HEAP_PROPERTIES HeapProps = CD3DX12_HEAP_PROPERTIES(D3D12_HEAP_TYPE_DEFAULT, GetGPUMask().GetNative(), GetVisibilityMask().GetNative());
-	Desc.Alignment = 0;
-	VERIFYD3D12RESULT(RetCode = Adapter->CreateCommittedResource(Desc, GetGPUMask(), HeapProps, InitialState, ClearValue, &NewResource, Name, false));
-
-	TextureLocation.AsStandAlone(NewResource);
-	return RetCode;
-}
 #else
+
 FD3D12TextureAllocator::FD3D12TextureAllocator(FD3D12Device* Device,
 	FRHIGPUMask VisibleNodes,
 	const FString& Name,
@@ -2461,225 +2140,3 @@ void* FD3D12FastConstantAllocator::Allocate(uint32 Bytes, FD3D12ResourceLocation
 
 	return OutLocation.GetMappedBaseAddress();
 }
-
-FD3D12SegHeap* FD3D12SegList::CreateBackingHeap(
-	FD3D12Device* Parent,
-	FRHIGPUMask VisibleNodeMask,
-	D3D12_HEAP_TYPE HeapType,
-	D3D12_HEAP_FLAGS HeapFlags)
-{
-	// CS can be unlocked at this point and re-locked before adding it to FreeHeaps
-	// but doing so may cause multiple heaps to be created
-	ID3D12Heap* D3DHeap = nullptr;
-	D3D12_HEAP_DESC Desc = {};
-	Desc.SizeInBytes = HeapSize;
-	Desc.Properties = CD3DX12_HEAP_PROPERTIES(HeapType, Parent->GetGPUMask().GetNative(), VisibleNodeMask.GetNative());
-	Desc.Flags = HeapFlags;
-	if (Parent->GetParentAdapter()->IsHeapNotZeroedSupported())
-	{
-		Desc.Flags |= FD3D12_HEAP_FLAG_CREATE_NOT_ZEROED;
-	}
-
-	VERIFYD3D12RESULT(Parent->GetDevice()->CreateHeap(&Desc, IID_PPV_ARGS(&D3DHeap)));
-
-	FD3D12SegHeap* Ret = new FD3D12SegHeap(Parent, VisibleNodeMask, D3DHeap, HeapSize, this, FreeHeaps.Num());
-	FreeHeaps.Add(Ret);
-	return Ret;
-}
-
-FD3D12SegListAllocator::FD3D12SegListAllocator(
-	FD3D12Device* Parent,
-	FRHIGPUMask VisibilityMask,
-	D3D12_HEAP_TYPE InHeapType,
-	D3D12_HEAP_FLAGS InHeapFlags,
-	uint32 InMinPoolSize,
-	uint32 InMinNumToPool,
-	uint32 InMaxPoolSize)
-	: FD3D12DeviceChild(Parent)
-	, FD3D12MultiNodeGPUObject(Parent->GetGPUMask(), VisibilityMask)
-	, HeapType(InHeapType)
-	, HeapFlags(InHeapFlags)
-	, MinPoolSize(InMinPoolSize)
-	, MinNumToPool(InMinNumToPool)
-	, MaxPoolSize(InMaxPoolSize)
-#if D3D12RHI_SEGLIST_ALLOC_TRACK_WASTAGE
-	, TotalBytesRequested(0)
-#endif
-{}
-
-void FD3D12SegListAllocator::Deallocate(
-	FD3D12Resource* PlacedResource,
-	uint32 Offset,
-	uint32 SizeInBytes)
-{
-	FD3D12Device* Device = this->GetParentDevice();
-	FD3D12Adapter* Adapter = Device->GetParentAdapter();
-	uint64 CurFenceValue = Adapter->GetFrameFence().GetNextFenceToSignal();
-	{
-		FScopeLock Lock(&DeferredDeletionCS);
-
-		check(FenceValues.Num() == DeferredDeletionQueue.Num());
-		check(!FenceValues.Num() || FenceValues.Last() <= CurFenceValue);
-
-		int32 LastIdx = FenceValues.Num() - 1;
-		if (LastIdx < 0 || FenceValues[LastIdx] != CurFenceValue)
-		{
-			++LastIdx;
-			FenceValues.Add(CurFenceValue);
-			DeferredDeletionQueue.AddDefaulted();
-		}
-		new (DeferredDeletionQueue[LastIdx]) FRetiredBlock(PlacedResource, Offset, SizeInBytes);
-	}
-
-	DEC_DWORD_STAT(STAT_D3D12TextureAllocatorCount);
-}
-
-template <typename AllocX, typename AllocY>
-void FD3D12SegListAllocator::FreeRetiredBlocks(TArray<TArray<FRetiredBlock, AllocX>, AllocY>& PendingDeletes)
-{
-	for (int32 Y = 0; Y < PendingDeletes.Num(); ++Y)
-	{
-		TArray<FRetiredBlock>& RetiredBlocks = PendingDeletes[Y];
-		for (int32 X = 0; X < RetiredBlocks.Num(); ++X)
-		{
-			FRetiredBlock& Block = RetiredBlocks[X];
-			if (ensureAlwaysMsgf(Block.PlacedResource->GetRefCount() == 1, TEXT("Invalid refcount while releasing %s"), *Block.PlacedResource->GetName().ToString()))
-			{
-				FD3D12SegHeap* BackingHeap = static_cast<FD3D12SegHeap*>(Block.PlacedResource->GetHeap());
-				Block.PlacedResource->Release();
-				FD3D12SegList* Owner = BackingHeap->OwnerList;
-				check(!!Owner);
-				Owner->FreeBlock(BackingHeap, Block.Offset);
-				OnFree(Block.Offset, BackingHeap, Block.ResourceSize);
-			}
-		}
-	}
-}
-
-void FD3D12SegListAllocator::CleanUpAllocations()
-{
-	TArray<TArray<FRetiredBlock>, TInlineAllocator<1>> PendingDeletes;
-	{
-		int32 NumToRemove = 0;
-		FD3D12Device* Device = this->GetParentDevice();
-		FD3D12Adapter* Adapter = Device->GetParentAdapter();
-		FD3D12ManualFence& FrameFence = Adapter->GetFrameFence();
-
-		FScopeLock Lock(&DeferredDeletionCS);
-
-		for (int32 Idx = 0; Idx < DeferredDeletionQueue.Num(); ++Idx)
-		{
-			if (FrameFence.IsFenceComplete(FenceValues[Idx], /* bUpdateCachedFenceValue */ false))
-			{
-				++NumToRemove;
-				PendingDeletes.Add(MoveTemp(DeferredDeletionQueue[Idx]));
-			}
-			else
-			{
-				break;
-			}
-		}
-		if (!!NumToRemove)
-		{
-			FenceValues.RemoveAt(0, NumToRemove);
-			DeferredDeletionQueue.RemoveAt(0, NumToRemove);
-		}
-	}
-	FreeRetiredBlocks(PendingDeletes);
-}
-
-void FD3D12SegListAllocator::Destroy()
-{
-	{
-		FScopeLock Lock(&DeferredDeletionCS);
-		check(FenceValues.Num() == DeferredDeletionQueue.Num());
-		FreeRetiredBlocks(DeferredDeletionQueue);
-		FenceValues.Empty();
-		DeferredDeletionQueue.Empty();
-		VerifyEmpty();
-	}
-	{
-		FRWScopeLock Lock(SegListsRWLock, SLT_Write);
-		for (auto& Pair : SegLists)
-		{
-			FD3D12SegList*& SegList = Pair.Value;
-			check(!!SegList);
-			delete SegList;
-			SegList = nullptr;
-		}
-		SegLists.Empty();
-	}
-}
-#if D3D12RHI_SEGLIST_ALLOC_TRACK_WASTAGE
-void FD3D12SegListAllocator::VerifyEmpty()
-{
-	FScopeLock Lock(&SegListTrackedAllocationCS);
-	if(SegListTrackedAllocations.Num() != 0)
-	{
-		UE_LOG(LogD3D12RHI, Warning, TEXT("Dumping leaked SegListAllocations\n"));
-		for (FD3D12SegListAllocatorLeakTrack& LeakTrack : SegListTrackedAllocations)
-		{
-			DumpStack(LeakTrack);
-		}
-	}
-
-	ensureMsgf(TotalBytesRequested == 0,
-		TEXT("FD3D12SegListAllocator contains %lld allocated bytes but is expected to be empty. This likely means a memory leak. Use d3d12.SegListTrackLeaks=1 CVar to print allocations to the log."),
-		(uint64)TotalBytesRequested);
-}
-
-
-void FD3D12SegListAllocator::DumpStack(const FD3D12SegListAllocatorLeakTrack& LeakTrack)
-{
-	UE_LOG(LogD3D12RHI, Warning, TEXT("Leaking Allocation Heap %p Offset %d\nStack Dump\n"), LeakTrack.Heap, LeakTrack.Offset);
-	for(uint32 Index = 0; Index < LeakTrack.StackDepth; ++Index)
-	{
-		const size_t STRING_SIZE = 16 * 1024;
-		ANSICHAR StackTrace[STRING_SIZE];
-		StackTrace[0] = 0;
-		FPlatformStackWalk::ProgramCounterToHumanReadableString(Index, LeakTrack.Stack[Index], StackTrace, STRING_SIZE, 0);
-		UE_LOG(LogD3D12RHI, Warning, TEXT("%d %S\n"), Index, StackTrace);
-	}
-}
-
-void FD3D12SegListAllocator::OnAlloc(uint32 Offset, void* Heap, uint32 Size)
-{
-	TotalBytesRequested += Size;
-
-	if(GD3D12SegListTrackLeaks == 0)
-		return;
-	FD3D12SegListAllocatorLeakTrack LeakTrack;
-	LeakTrack.Offset = Offset;
-	LeakTrack.Heap = Heap;
-	LeakTrack.Size = Size;
-	LeakTrack.StackDepth = FPlatformStackWalk::CaptureStackBackTrace(&LeakTrack.Stack[0], D3D12RHI_SEGLIST_ALLOC_TRACK_LEAK_STACK_DEPTH);
-
-	FScopeLock Lock(&SegListTrackedAllocationCS);
-	check(!SegListTrackedAllocations.Contains(LeakTrack));
-	SegListTrackedAllocations.Add(LeakTrack);
-}
-void FD3D12SegListAllocator::OnFree(uint32 Offset, void* Heap, uint32 Size)
-{
-	TotalBytesRequested -= Size;
-	if (GD3D12SegListTrackLeaks == 0)
-		return;
-
-	FD3D12SegListAllocatorLeakTrack LeakTrack;
-	LeakTrack.Offset = Offset;
-	LeakTrack.Heap = Heap;
-	FScopeLock Lock(&SegListTrackedAllocationCS);
-	FD3D12SegListAllocatorLeakTrack* Element = SegListTrackedAllocations.Find(LeakTrack);
-	check(Element); // element being freed was not found.
-	if(Element->Size != Size)
-	{
-		UE_LOG(LogD3D12RHI, Warning, TEXT("Mismatched alloc/free size %d != %d, %p/%08x"), Element->Size, Size, Element->Heap, Element->Offset);
-		DumpStack(*Element);
-		check(0); //element being freed had incorrect size. 
-	}
-	SegListTrackedAllocations.Remove(LeakTrack);
-	check(!SegListTrackedAllocations.Contains(LeakTrack));
-}
-#endif
-
-
-
diff --git a/Engine/Source/Runtime/D3D12RHI/Private/D3D12Allocation.h b/Engine/Source/Runtime/D3D12RHI/Private/D3D12Allocation.h
--- a/Engine/Source/Runtime/D3D12RHI/Private/D3D12Allocation.h
+++ b/Engine/Source/Runtime/D3D12RHI/Private/D3D12Allocation.h
@@ -12,74 +12,18 @@ D3D12Allocation.h: A Collection of allocators
 
 class FD3D12ConstantBufferView;
 
-#define SUB_ALLOCATED_DEFAULT_ALLOCATIONS	1
-
-// Enable pool allocator by default only on Windows until all initial issues have been resolved.
-// Xbox is currently using the buddy allocator with a lot smaller pool sizes so the alignment waste is a lot less
-// It also has less of a problem with committed resource allocations
 #if !defined(USE_BUFFER_POOL_ALLOCATOR)
-#define USE_BUFFER_POOL_ALLOCATOR (PLATFORM_WINDOWS)
+#error "USE_BUFFER_POOL_ALLOCATOR is not defined"
 #endif
 
 #if !defined(USE_TEXTURE_POOL_ALLOCATOR)
-#define USE_TEXTURE_POOL_ALLOCATOR (PLATFORM_WINDOWS)
+#error "USE_TEXTURE_POOL_ALLOCATOR is not defined"
 #endif
 
-// Segregated free list texture allocator
-// Description:
-// - Binned read-only texture allocation based on sizes
-// Suggestions:
-// - You can check memory wastage using "stat d3d12rhi" in a dev build
-// - Tune d3d12.ReadOnlyTextureAllocator.MinPoolSize/MinNumToPool/MaxPoolSize
-//   according to video memory budget
-// - Memory overhead is slightly over 200 MB in internal tests but consider
-//   adjusting above cvars or disabling if it fails your use case
-// - The purpose of this allocator is pooling texture allocations because
-//   creating committed resources is slow on PC. But if committed resource
-//   creation ever becomes fast, there is no need for this allocator
-// Internal test statistics (11/6/2018):
-// - Average read-only texture alloc time reduced from ~420 us to ~72 us
-// - Number of allocations over 1 ms reduced from 8145 to 504 (from 14.76% to
-//   0.92%) over a 17 minutes 11 seconds game replay
-// - Peak memory overhead was ~207 MB (from 2666.58 MB to 2874.08 MB)
-// TODO:
-// - Defragmentation support
-#define D3D12RHI_SEGREGATED_TEXTURE_ALLOC (PLATFORM_WINDOWS)
-#define D3D12RHI_SEGLIST_ALLOC_TRACK_WASTAGE (!(UE_BUILD_TEST || UE_BUILD_SHIPPING))
-
-#define D3D12RHI_SEGLIST_ALLOC_TRACK_LEAK_STACK_DEPTH 12
-
 #define D3D12RHI_TRACK_DETAILED_STATS (PLATFORM_WINDOWS && !(UE_BUILD_TEST || UE_BUILD_SHIPPING))
 
-struct FD3D12SegListAllocatorLeakTrack
-{
-	uint32 Offset;
-	void*  Heap;
-	uint32 Size;
-	uint32 StackDepth;
-	uint64 Stack[D3D12RHI_SEGLIST_ALLOC_TRACK_LEAK_STACK_DEPTH];
-	bool operator==(const FD3D12SegListAllocatorLeakTrack& Other) const
-	{
-		return Offset == Other.Offset && Heap == Other.Heap;
-	}
-};
-
-FORCEINLINE uint32 GetTypeHash(const FD3D12SegListAllocatorLeakTrack& S)
-{
-	uint32 Prime0 = 0xa6c70167;
-	uint32 Prime1 = 0x5d18b207;
-	uint32 Prime2 = 0xd0a489f9;
-	uint32 Value0 = (uint32)(((uint64)S.Heap) >> 32);
-	uint32 Value1 = (uint32)(uint64)S.Heap;
-	uint32 Value2 = S.Offset;
-	return Value0 * Prime0 + Value1 * Prime1 + Value2 * Prime2;
-}
-
-
 const uint32 kD3D12ManualSubAllocationAlignment = 256;
 
-class FD3D12SegList;
-
 class FD3D12ResourceAllocator : public FD3D12DeviceChild, public FD3D12MultiNodeGPUObject
 {
 public:
@@ -298,74 +242,6 @@ private:
 	HeapId TraceHeapId;
 };
 
-//-----------------------------------------------------------------------------
-//	Bucket Allocator
-//-----------------------------------------------------------------------------
-// Resources are allocated from buckets, which are just a collection of resources of a particular size.
-// Blocks can be an entire resource or a sub allocation from a resource.
-class FD3D12BucketAllocator : public FD3D12ResourceAllocator
-{
-public:
-
-	FD3D12BucketAllocator(FD3D12Device* ParentDevice,
-		FRHIGPUMask VisibleNodes,
-		const FD3D12ResourceInitConfig& InInitConfig,
-		const FString& Name,
-		uint64 InBlockRetentionFrameCount);
-
-	bool TryAllocate(uint32 SizeInBytes, uint32 Alignment, FD3D12ResourceLocation& ResourceLocation);
-
-	void Deallocate(FD3D12ResourceLocation& ResourceLocation);
-
-	void Initialize();
-
-	void Destroy();
-
-	void CleanUpAllocations(uint64 InFrameLag);
-
-	void DumpAllocatorStats(class FOutputDevice& Ar);
-	void UpdateMemoryStats(uint32& IOMemoryAllocated, uint32& IOMemoryUsed, uint32& IOMemoryFree, uint32& IOAlignmentWaste, uint32& IOAllocatedPageCount, uint32& IOFullPageCount) {}
-
-	void ReleaseAllResources();
-
-	void Reset();
-
-private:
-
-	static uint32 FORCEINLINE BucketFromSize(uint32 size, uint32 bucketShift)
-	{
-		uint32 bucket = FMath::CeilLogTwo(size);
-		bucket = bucket < bucketShift ? 0 : bucket - bucketShift;
-		return bucket;
-	}
-
-	static uint32 FORCEINLINE BlockSizeFromBufferSize(uint32 bufferSize, uint32 bucketShift)
-	{
-		const uint32 minSize = 1 << bucketShift;
-		return bufferSize > minSize ? FMath::RoundUpToPowerOfTwo(bufferSize) : minSize;
-	}
-
-#if SUB_ALLOCATED_DEFAULT_ALLOCATIONS
-	static const uint32 MIN_HEAP_SIZE = 256 * 1024;
-#else
-	static const uint32 MIN_HEAP_SIZE = 64 * 1024;
-#endif
-
-	static const uint32 BucketShift = 6;
-	static const uint32 NumBuckets = 22; // bucket resource sizes range from 64 to 2^28 
-
-	FThreadsafeQueue<FD3D12BlockAllocatorPrivateData> AvailableBlocks[NumBuckets];
-	FThreadsafeQueue<FD3D12BlockAllocatorPrivateData> ExpiredBlocks;
-	TArray<FD3D12Resource*> SubAllocatedResources;// keep a list of the sub-allocated resources so that they may be cleaned up
-
-	// This frame count value helps makes sure that we don't delete resources too soon. If resources are deleted too soon,
-	// we can get in a loop the heap allocator will be constantly deleting and creating resources every frame which
-	// results in CPU stutters. DynamicRetentionFrameCount was tested and set to a value that appears to be adequate for
-	// creating a stable state on the Infiltrator demo.
-	const uint64 BlockRetentionFrameCount;
-};
-
-
 //-----------------------------------------------------------------------------
 //	FD3D12UploadHeapAllocator
 //-----------------------------------------------------------------------------
@@ -564,369 +440,6 @@ private:
 	uint32 PageSize;
 };
 
-//-----------------------------------------------------------------------------
-//	FD3D12SegListAllocator
-//-----------------------------------------------------------------------------
-
-class FD3D12SegHeap : public FD3D12Heap
-{
-private:
-	FD3D12SegHeap(
-		FD3D12Device* Parent,
-		FRHIGPUMask VisibileNodeMask,
-		ID3D12Heap* NewHeap,
-		uint64 HeapSize,
-		FD3D12SegList* Owner,
-		uint32 Idx) :
-		FD3D12Heap(Parent, VisibileNodeMask),
-		OwnerList(Owner),
-		ArrayIdx(Idx),
-		FirstFreeOffset(0)
-	{
-		this->SetHeap(NewHeap, TEXT("SegListHeap"));
-		BeginTrackingResidency(HeapSize);
-	}
-
-	virtual ~FD3D12SegHeap() = default;
-
-	FD3D12SegHeap(const FD3D12SegHeap&) = delete;
-	FD3D12SegHeap(FD3D12SegHeap&&) = delete;
-
-	FD3D12SegHeap& operator=(const FD3D12SegHeap&) = delete;
-	FD3D12SegHeap& operator=(FD3D12SegHeap&&) = delete;
-
-	bool IsArrayIdxValid() const { return ArrayIdx >= 0; }
-
-	bool IsFull(uint32 HeapSize) const
-	{
-		check(FirstFreeOffset <= HeapSize);
-		return !FreeBlockOffsets.Num() && FirstFreeOffset == HeapSize;
-	}
-
-	bool IsEmpty(uint32 BlockSize) const
-	{
-		return FreeBlockOffsets.Num() * BlockSize == FirstFreeOffset;
-	}
-
-	// @return - In-heap offset of the allocated block
-	uint32 AllocateBlock(uint32 BlockSize)
-	{
-		if (!FreeBlockOffsets.Num())
-		{
-			uint32 Ret = FirstFreeOffset;
-			FirstFreeOffset += BlockSize;
-			return Ret;
-		}
-		else
-		{
-			return FreeBlockOffsets.Pop();
-		}
-	}
-
-	TArray<uint32> FreeBlockOffsets;
-	FD3D12SegList* OwnerList;
-	int32 ArrayIdx;
-	uint32 FirstFreeOffset;
-
-	friend class FD3D12SegList;
-	friend class FD3D12SegListAllocator;
-};
-
-class FD3D12SegList
-{
-private:
-	FD3D12SegList(uint32 InBlockSize, uint32 InHeapSize)
-		: BlockSize(InBlockSize)
-		, HeapSize(InHeapSize)
-#if D3D12RHI_SEGLIST_ALLOC_TRACK_WASTAGE
-		, TotalBytesAllocated(0)
-#endif
-	{
-		check(!(HeapSize % BlockSize));
-		check(HeapSize / BlockSize > 1);
-	}
-
-	~FD3D12SegList()
-	{
-		FScopeLock Lock(&CS);
-		check(!!BlockSize);
-		check(!!HeapSize);
-		for (const auto& Heap : FreeHeaps)
-		{
-			ensure(Heap->GetRefCount() == 1);
-		}
-	}
-
-	FD3D12SegList(const FD3D12SegList&) = delete;
-	FD3D12SegList(FD3D12SegList&&) = delete;
-
-	FD3D12SegList& operator=(const FD3D12SegList&) = delete;
-	FD3D12SegList& operator=(FD3D12SegList&&) = delete;
-
-	// @return - In-heap offset of the allocated block
-	uint32 AllocateBlock(
-		FD3D12Device* Device,
-		FRHIGPUMask VisibleNodeMask,
-		D3D12_HEAP_TYPE HeapType,
-		D3D12_HEAP_FLAGS HeapFlags,
-		TRefCountPtr<FD3D12SegHeap>& OutHeap)
-	{
-		FScopeLock Lock(&CS);
-		uint32 Offset;
-
-		if (!!FreeHeaps.Num())
-		{
-			const int32 LastHeapIdx = FreeHeaps.Num() - 1;
-			OutHeap = FreeHeaps[LastHeapIdx];
-			Offset = OutHeap->AllocateBlock(BlockSize);
-			check(Offset <= HeapSize - BlockSize);
-			if (OutHeap->IsFull(HeapSize))
-			{
-				// Heap is full
-				OutHeap->ArrayIdx = INDEX_NONE;
-				FreeHeaps.RemoveAt(LastHeapIdx);
-			}
-		}
-		else
-		{
-			OutHeap = CreateBackingHeap(Device, VisibleNodeMask, HeapType, HeapFlags);
-			Offset = OutHeap->AllocateBlock(BlockSize);
-#if D3D12RHI_SEGLIST_ALLOC_TRACK_WASTAGE
-			TotalBytesAllocated += HeapSize;
-#endif
-		}
-		return Offset;
-	}
-
-	// Deferred deletion is handled by FD3D12SegListAllocator
-	void FreeBlock(FD3D12SegHeap* RESTRICT Heap, uint32 Offset)
-	{
-		FScopeLock Lock(&CS);
-
-		check(!(Offset % BlockSize));
-		check(Offset <= HeapSize - BlockSize);
-		check(this == Heap->OwnerList);
-
-		const bool bFull = Heap->IsFull(HeapSize);
-		Heap->FreeBlockOffsets.Add(Offset);
-
-		if (bFull)
-		{
-			// Heap was full
-			check(!Heap->IsArrayIdxValid());
-			Heap->ArrayIdx = FreeHeaps.Add(Heap);
-		}
-		else if (Heap->IsEmpty(BlockSize))
-		{
-			// Heap is empty
-			check(Heap->GetRefCount() == 1);
-			check(Heap->IsArrayIdxValid());
-			check(FreeHeaps.Num() > Heap->ArrayIdx);
-			const int32 Idx = Heap->ArrayIdx;
-			const int32 LastIdx = FreeHeaps.Num() - 1;
-			FreeHeaps.RemoveAtSwap(Idx);
-			if (Idx != LastIdx)
-			{
-				FreeHeaps[Idx]->ArrayIdx = Idx;
-			}
-#if D3D12RHI_SEGLIST_ALLOC_TRACK_WASTAGE
-			TotalBytesAllocated -= HeapSize;
-#endif
-		}
-	}
-
-	FD3D12SegHeap* CreateBackingHeap(
-		FD3D12Device* Parent,
-		FRHIGPUMask VisibleNodeMask,
-		D3D12_HEAP_TYPE HeapType,
-		D3D12_HEAP_FLAGS HeapFlags);
-
-	TArray<TRefCountPtr<FD3D12SegHeap>> FreeHeaps;
-	FCriticalSection CS;
-	uint32 BlockSize;
-	uint32 HeapSize;
-
-#if D3D12RHI_SEGLIST_ALLOC_TRACK_WASTAGE
-	uint64 TotalBytesAllocated;
-#endif
-
-	friend class FD3D12SegListAllocator;
-};
-
-#if !D3D12RHI_SEGLIST_ALLOC_TRACK_WASTAGE
-static_assert(sizeof(FD3D12SegList) <= 64, "Try to make it fit in a single cacheline");
-#endif
-
-class FD3D12SegListAllocator : public FD3D12DeviceChild, public FD3D12MultiNodeGPUObject
-{
-public:
-	static constexpr uint32 InvalidOffset = 0xffffffff;
-
-	FD3D12SegListAllocator(
-		FD3D12Device* Parent,
-		FRHIGPUMask VisibilityMask,
-		D3D12_HEAP_TYPE InHeapType,
-		D3D12_HEAP_FLAGS InHeapFlags,
-		uint32 InMinPoolSize,
-		uint32 InMinNumToPool,
-		uint32 InMaxPoolSize);
-
-	~FD3D12SegListAllocator()
-	{
-		check(!SegLists.Num());
-		check(!FenceValues.Num());
-		check(!DeferredDeletionQueue.Num());
-		VerifyEmpty();
-	}
-
-	FD3D12SegListAllocator(const FD3D12SegListAllocator&) = delete;
-	FD3D12SegListAllocator(FD3D12SegListAllocator&&) = delete;
-
-	FD3D12SegListAllocator& operator=(const FD3D12SegListAllocator&) = delete;
-	FD3D12SegListAllocator& operator=(FD3D12SegListAllocator&&) = delete;
-
-	uint32 Allocate(uint64 SizeInBytes, uint64 Alignment, TRefCountPtr<FD3D12SegHeap>& OutHeap)
-	{
-		check(!(Alignment & Alignment - 1));
-
-		const uint64 BlockSize = CalculateBlockSize(SizeInBytes, Alignment);
-		if (ShouldPool(BlockSize))
-		{
-			FD3D12SegList* SegList;
-			{
-				FRWScopeLock Lock(SegListsRWLock, SLT_ReadOnly);
-				FD3D12SegList** SegListPtr = SegLists.Find(BlockSize);
-				SegList = !!SegListPtr ? *SegListPtr : nullptr;
-			}
-			if (!SegList)
-			{
-				const uint32 HeapSize = CalculateHeapSize(BlockSize);
-				{
-					FRWScopeLock Lock(SegListsRWLock, SLT_Write);
-					FD3D12SegList** SegListPtr = SegLists.Find(BlockSize);
-					SegList = !!SegListPtr ?
-						*SegListPtr :
-						SegLists.Add(BlockSize, new FD3D12SegList(BlockSize, HeapSize));
-				}
-			}
-			check(!!SegList);
-			uint32 Ret = SegList->AllocateBlock(
-				this->GetParentDevice(),
-				this->GetVisibilityMask(),
-				HeapType,
-				HeapFlags,
-				OutHeap);
-			check(Ret != InvalidOffset);
-			OnAlloc(Ret, OutHeap.GetReference(), SizeInBytes);
-			return Ret;
-		}
-		OutHeap = nullptr;
-		return InvalidOffset;
-	}
-
-	void Deallocate(FD3D12Resource* PlacedResource, uint32 Offset, uint32 SizeInBytes);
-
-	void CleanUpAllocations();
-
-	void Destroy();
-
-	bool GetMemoryStats(uint64& OutTotalAllocated, uint64& OutTotalUnused) const
-	{
-#if D3D12RHI_SEGLIST_ALLOC_TRACK_WASTAGE
-		FScopeLock LockCS(&DeferredDeletionCS);
-		FRWScopeLock LockRW(SegListsRWLock, SLT_Write);
-
-		OutTotalAllocated = 0;
-		for (const auto& Pair : SegLists)
-		{
-			const FD3D12SegList* SegList = Pair.Value;
-			OutTotalAllocated += SegList->TotalBytesAllocated;
-		}
-		OutTotalUnused = OutTotalAllocated - TotalBytesRequested.Load(EMemoryOrder::Relaxed);
-		return true;
-#else
-		return false;
-#endif
-	}
-
-private:
-	struct FRetiredBlock
-	{
-		// FD3D12Resource knows which heap it is from
-		FD3D12Resource* PlacedResource;
-		uint32 Offset;
-		uint32 ResourceSize;
-
-		FRetiredBlock(
-			FD3D12Resource* InResource,
-			uint32 InOffset,
-			uint32 InResourceSize) :
-			PlacedResource(InResource),
-			Offset(InOffset),
-			ResourceSize(InResourceSize)
-		{}
-	};
-
-	static constexpr uint64 CalculateBlockSize(uint64 SizeInBytes, uint64 Alignment)
-	{
-		return (SizeInBytes + Alignment - 1) & ~(Alignment - 1);
-	}
-
-	bool ShouldPool(uint64 BlockSize) const
-	{
-		return BlockSize * 2 <= MaxPoolSize;
-	}
-
-	uint32 CalculateHeapSize(uint32 BlockSize) const
-	{
-		check(MinPoolSize + BlockSize - 1 > MinPoolSize);
-		uint32 NumPooled = (MinPoolSize + BlockSize - 1) / BlockSize;
-		if (NumPooled < MinNumToPool)
-		{
-			NumPooled = MinNumToPool;
-		}
-		const uint32 MaxNumPooled = MaxPoolSize / BlockSize;
-		if (NumPooled > MaxNumPooled)
-		{
-			NumPooled = MaxNumPooled;
-		}
-		check(NumPooled > 1);
-		check(NumPooled * BlockSize >= MinPoolSize);
-		check(NumPooled * BlockSize <= MaxPoolSize);
-		return NumPooled * BlockSize;
-	}
-
-	template <typename AllocX, typename AllocY>
-	void FreeRetiredBlocks(TArray<TArray<FRetiredBlock, AllocX>, AllocY>& PendingDeletes);
-
-	TMap<uint32, FD3D12SegList*> SegLists;
-	TArray<uint64> FenceValues;
-	TArray<TArray<FRetiredBlock>> DeferredDeletionQueue;
-	mutable FRWLock SegListsRWLock;
-	mutable FCriticalSection DeferredDeletionCS;
-	const D3D12_HEAP_TYPE HeapType;
-	const D3D12_HEAP_FLAGS HeapFlags;
-	const uint32 MinPoolSize;
-	const uint32 MinNumToPool;
-	const uint32 MaxPoolSize;
-
-#if D3D12RHI_SEGLIST_ALLOC_TRACK_WASTAGE
-	TAtomic<uint64> TotalBytesRequested;
-	FCriticalSection SegListTrackedAllocationCS;
-	TSet<FD3D12SegListAllocatorLeakTrack> SegListTrackedAllocations;
-
-
-	void DumpStack(const FD3D12SegListAllocatorLeakTrack& LeakTrack);
-	void OnAlloc(uint32 Offset, void* Heap, uint32 Size);
-	void OnFree(uint32 Offset, void* Heap, uint32 Size);
-	void VerifyEmpty();
-#else
-	void OnAlloc(uint32 Offset, void* Heap, uint32 Size) {}
-	void OnFree(uint32 Offset, void* Heap, uint32 Size) {}
-	void VerifyEmpty(){}
-#endif
-};
-
 //-----------------------------------------------------------------------------
 //	FD3D12TextureAllocator
 //-----------------------------------------------------------------------------
@@ -965,33 +478,6 @@ private:
 	FD3D12PoolAllocator* PoolAllocators[(int)EPoolType::Count];
 	HeapId TraceHeapId;
 };
-#elif D3D12RHI_SEGREGATED_TEXTURE_ALLOC
-class FD3D12TextureAllocatorPool : public FD3D12DeviceChild, public FD3D12MultiNodeGPUObject
-{
-public:
-	FD3D12TextureAllocatorPool(FD3D12Device* Device, FRHIGPUMask VisibilityNode);
-
-	HRESULT AllocateTexture(FD3D12ResourceDesc Desc, const D3D12_CLEAR_VALUE* ClearValue, EPixelFormat UEFormat, FD3D12ResourceLocation& TextureLocation, const D3D12_RESOURCE_STATES InitialState, const TCHAR* Name);
-
-	void BeginFrame(FRHICommandListBase& RHICmdList) {}
-	void CleanUpAllocations()
-	{
-		ReadOnlyTexturePool.CleanUpAllocations();
-	}
-
-	void Destroy()
-	{
-		ReadOnlyTexturePool.Destroy();
-	}
-
-	bool GetMemoryStats(uint64& OutTotalAllocated, uint64& OutTotalUnused) const
-	{
-		return ReadOnlyTexturePool.GetMemoryStats(OutTotalAllocated, OutTotalUnused);
-	}
-
-private:
-	FD3D12SegListAllocator ReadOnlyTexturePool;
-};
 #else
 class FD3D12TextureAllocator : public FD3D12MultiBuddyAllocator
 {
diff --git a/Engine/Source/Runtime/D3D12RHI/Private/D3D12CommandContext.cpp b/Engine/Source/Runtime/D3D12RHI/Private/D3D12CommandContext.cpp
--- a/Engine/Source/Runtime/D3D12RHI/Private/D3D12CommandContext.cpp
+++ b/Engine/Source/Runtime/D3D12RHI/Private/D3D12CommandContext.cpp
@@ -722,15 +722,6 @@ void FD3D12DynamicRHI::UpdateMemoryStats()
 		uint64 MaxTexAllocWastage = 0;
 		for (FD3D12Device* Device : Adapter->GetDevices())
 		{
-#if D3D12RHI_SEGREGATED_TEXTURE_ALLOC && D3D12RHI_SEGLIST_ALLOC_TRACK_WASTAGE
-			uint64 TotalAllocated;
-			uint64 TotalUnused;
-			Device->GetTextureAllocator().GetMemoryStats(TotalAllocated, TotalUnused);
-			MaxTexAllocWastage = FMath::Max(MaxTexAllocWastage, TotalUnused);
-			SET_MEMORY_STAT(STAT_D3D12TextureAllocatorAllocated, TotalAllocated);
-			SET_MEMORY_STAT(STAT_D3D12TextureAllocatorUnused, TotalUnused);
-#endif
-
 			Device->GetDefaultBufferAllocator().UpdateMemoryStats();
 			Adapter->GetUploadHeapAllocator(Device->GetGPUIndex()).UpdateMemoryStats();
 		}
diff --git a/Engine/Source/Runtime/D3D12RHI/Private/D3D12RHIPrivate.h b/Engine/Source/Runtime/D3D12RHI/Private/D3D12RHIPrivate.h
--- a/Engine/Source/Runtime/D3D12RHI/Private/D3D12RHIPrivate.h
+++ b/Engine/Source/Runtime/D3D12RHI/Private/D3D12RHIPrivate.h
@@ -396,8 +396,10 @@ public:
 #if D3D12RHI_SUPPORTS_UAV_BACKBUFFER
 	virtual FUnorderedAccessViewRHIRef RHIGetViewportBackBufferUAV(FRHIViewport* Viewport) final override;
 #endif
+#if D3D12RHI_SUPPORTS_ALIASED_RHI_RESOURCES
 	virtual void RHIAliasTextureResources(FTextureRHIRef& DestTexture, FTextureRHIRef& SrcTexture) final override;
 	virtual FTextureRHIRef RHICreateAliasedTexture(FTextureRHIRef& SourceTexture) final override;
+#endif
 	virtual void RHIGetDisplaysInformation(FDisplayInformationArray& OutDisplayInformation) final override;
 	virtual uint64 RHIComputeStatePrecachePSOHash(const FGraphicsPipelineStateInitializer& Initializer) final override;
 	virtual uint64 RHIComputePrecachePSOHash(const FGraphicsPipelineStateInitializer& Initializer) final override;
@@ -638,7 +640,9 @@ public:
 protected:
 
 	FD3D12Texture* CreateTextureFromResource(bool bTextureArray, bool bCubeTexture, EPixelFormat Format, ETextureCreateFlags TexCreateFlags, const FClearValueBinding& ClearValueBinding, ID3D12Resource* Resource);
+#if D3D12RHI_SUPPORTS_ALIASED_RHI_RESOURCES
 	FD3D12Texture* CreateAliasedD3D12Texture2D(FD3D12Texture* SourceTexture);
+#endif
 
 	/**
 	 * Gets the best supported MSAA settings from the provided MSAA count to check against.
diff --git a/Engine/Source/Runtime/D3D12RHI/Private/D3D12Resources.cpp b/Engine/Source/Runtime/D3D12RHI/Private/D3D12Resources.cpp
--- a/Engine/Source/Runtime/D3D12RHI/Private/D3D12Resources.cpp
+++ b/Engine/Source/Runtime/D3D12RHI/Private/D3D12Resources.cpp
@@ -1226,7 +1226,7 @@ void FD3D12ResourceLocation::TransferOwnership(FD3D12ResourceLocation& Destinati
 
 	// update tracked allocation
 #if !PLATFORM_WINDOWS && ENABLE_LOW_LEVEL_MEM_TRACKER
-	if (Source.GetType() == ResourceLocationType::eSubAllocation && Source.AllocatorType != AT_SegList)
+	if (Source.GetType() == ResourceLocationType::eSubAllocation)
 	{
 		FLowLevelMemTracker::Get().OnLowLevelAllocMoved(ELLMTracker::Default, Destination.GetAddressForLLMTracking(), Source.GetAddressForLLMTracking());
 	}
@@ -1239,7 +1239,8 @@ void FD3D12ResourceLocation::TransferOwnership(FD3D12ResourceLocation& Destinati
 	Source.Owner = SrcOwner;
 }
 
-void FD3D12ResourceLocation::Alias(FD3D12ResourceLocation & Destination, FD3D12ResourceLocation & Source)
+#if D3D12RHI_SUPPORTS_ALIASED_RHI_RESOURCES
+void FD3D12ResourceLocation::Alias(FD3D12ResourceLocation& Destination, FD3D12ResourceLocation& Source)
 {
 	// Should not be linked list allocated - otherwise internal linked list data needs to be updated as well in a threadsafe way
 	check(Source.GetAllocatorType() != FD3D12ResourceLocation::AT_Pool);
@@ -1254,6 +1255,7 @@ void FD3D12ResourceLocation::Alias(FD3D12ResourceLocation & Destination, FD3D12R
 	// Addref the source as another resource location references it
 	Source.GetResource()->AddRef();
 }
+#endif
 
 void FD3D12ResourceLocation::ReferenceNode(FD3D12Device* DestinationDevice, FD3D12ResourceLocation& Destination, FD3D12ResourceLocation& Source)
 {
@@ -1322,14 +1324,7 @@ void FD3D12ResourceLocation::ReleaseResource()
 	case ResourceLocationType::eSubAllocation:
 	{
 		check(Allocator != nullptr);
-		if (AllocatorType == AT_SegList)
-		{
-			SegListAllocator->Deallocate(
-				GetResource(),
-				GetSegListAllocatorPrivateData().Offset,
-				GetSize());
-		}
-		else if (AllocatorType == AT_Pool)
+		if (AllocatorType == AT_Pool)
 		{
 			// Unlink any aliases -- the contents of aliases are cleaned up separately elsewhere via iteration over
 			// the FD3D12LinkedAdapterObject.
diff --git a/Engine/Source/Runtime/D3D12RHI/Private/D3D12Resources.h b/Engine/Source/Runtime/D3D12RHI/Private/D3D12Resources.h
--- a/Engine/Source/Runtime/D3D12RHI/Private/D3D12Resources.h
+++ b/Engine/Source/Runtime/D3D12RHI/Private/D3D12Resources.h
@@ -22,7 +22,6 @@ class FD3D12Resource;
 class FD3D12StateCache;
 class FD3D12CommandListManager;
 class FD3D12CommandContext;
-class FD3D12SegListAllocator;
 class FD3D12PoolAllocator;
 struct FD3D12ComputePipelineState;
 struct FD3D12WorkGraphPipelineState;
@@ -597,32 +596,6 @@ struct FD3D12BuddyAllocatorPrivateData
 	}
 };
 
-struct FD3D12BlockAllocatorPrivateData
-{
-	uint64 FrameFence;
-	uint32 BucketIndex;
-	uint32 Offset;
-	FD3D12Resource* ResourceHeap;
-
-	void Init()
-	{
-		FrameFence = 0;
-		BucketIndex = 0;
-		Offset = 0;
-		ResourceHeap = nullptr;
-	}
-};
-
-struct FD3D12SegListAllocatorPrivateData
-{
-	uint32 Offset;
-
-	void Init()
-	{
-		Offset = 0;
-	}
-};
-
 struct FD3D12PoolAllocatorPrivateData
 {
 	FRHIPoolAllocationData PoolData;
@@ -656,7 +629,6 @@ public:
 	enum EAllocatorType : uint8
 	{
 		AT_Default, // FD3D12BaseAllocatorType
-		AT_SegList, // FD3D12SegListAllocator
 		AT_Pool,	// FD3D12PoolAllocator
 		AT_Unknown = 0xff
 	};
@@ -681,7 +653,6 @@ public:
 	void SetResource(FD3D12Resource* Value);
 
 	void SetAllocator       (FD3D12BaseAllocatorType* Value) { Allocator        = Value; AllocatorType = AT_Default; }
-	void SetSegListAllocator(FD3D12SegListAllocator*  Value) { SegListAllocator = Value; AllocatorType = AT_SegList; }
 	void SetPoolAllocator   (FD3D12PoolAllocator*     Value) { PoolAllocator    = Value; AllocatorType = AT_Pool;    }
 
 	void ClearAllocator() { Allocator = nullptr; AllocatorType = AT_Unknown; }
@@ -695,7 +666,6 @@ public:
 	ResourceLocationType               GetType                       () const { return Type;                                                 }
 	EAllocatorType                     GetAllocatorType              () const { return AllocatorType;                                        }
 	FD3D12BaseAllocatorType*           GetAllocator                  ()       { check(AT_Default == AllocatorType); return Allocator;        }
-	FD3D12SegListAllocator*            GetSegListAllocator           ()       { check(AT_SegList == AllocatorType); return SegListAllocator; }
 	FD3D12PoolAllocator*               GetPoolAllocator              ()       { check(AT_Pool == AllocatorType);    return PoolAllocator;    }
 	FD3D12Resource*                    GetResource                   () const { return UnderlyingResource;                                   }
 	void*                              GetMappedBaseAddress          () const { return MappedBaseAddress;                                    }
@@ -703,8 +673,6 @@ public:
 	uint64                             GetOffsetFromBaseOfResource   () const { return OffsetFromBaseOfResource;                             }
 	uint64                             GetSize                       () const { return Size;                                                 }
 	FD3D12BuddyAllocatorPrivateData&   GetBuddyAllocatorPrivateData  ()       { return AllocatorData.BuddyAllocatorPrivateData;              }
-	FD3D12BlockAllocatorPrivateData&   GetBlockAllocatorPrivateData  ()       { return AllocatorData.BlockAllocatorPrivateData;              }
-	FD3D12SegListAllocatorPrivateData& GetSegListAllocatorPrivateData()       { return AllocatorData.SegListAllocatorPrivateData;            }
 	FD3D12PoolAllocatorPrivateData&    GetPoolAllocatorPrivateData   ()       { return AllocatorData.PoolAllocatorPrivateData;               }
 
 	// Pool allocation specific functions
@@ -755,10 +723,13 @@ public:
 		SetGPUVirtualAddress(GPUBase + Offset);
 	}
 
+#if D3D12RHI_SUPPORTS_ALIASED_RHI_RESOURCES
 	// XR plugins alias textures so this allows 2+ resource locations to reference the same underlying
 	// resource. We should avoid this as much as possible as it requires expensive reference counting and
 	// it complicates the resource ownership model.
 	static void Alias(FD3D12ResourceLocation& Destination, FD3D12ResourceLocation& Source);
+#endif
+
 	static void ReferenceNode(FD3D12Device* NodeDevice, FD3D12ResourceLocation& Destination, FD3D12ResourceLocation& Source);
 	bool IsAliased() const
 	{
@@ -796,7 +767,6 @@ private:
 	union
 	{
 		FD3D12BaseAllocatorType* Allocator;
-		FD3D12SegListAllocator* SegListAllocator;
 		FD3D12PoolAllocator* PoolAllocator;
 	};
 
@@ -804,8 +774,6 @@ private:
 	union PrivateAllocatorData
 	{
 		FD3D12BuddyAllocatorPrivateData BuddyAllocatorPrivateData;
-		FD3D12BlockAllocatorPrivateData BlockAllocatorPrivateData;
-		FD3D12SegListAllocatorPrivateData SegListAllocatorPrivateData;
 		FD3D12PoolAllocatorPrivateData PoolAllocatorPrivateData;
 	} AllocatorData;
 
diff --git a/Engine/Source/Runtime/D3D12RHI/Private/D3D12Texture.cpp b/Engine/Source/Runtime/D3D12RHI/Private/D3D12Texture.cpp
--- a/Engine/Source/Runtime/D3D12RHI/Private/D3D12Texture.cpp
+++ b/Engine/Source/Runtime/D3D12RHI/Private/D3D12Texture.cpp
@@ -1962,6 +1962,7 @@ void FD3D12Texture::CreateViews(FD3D12Texture* FirstLinkedObject)
 }
 
 
+#if D3D12RHI_SUPPORTS_ALIASED_RHI_RESOURCES
 void FD3D12Texture::AliasResources(FD3D12Texture* Texture)
 {
 	// Alias the location, will perform an addref underneath
@@ -1982,6 +1983,7 @@ void FD3D12Texture::AliasResources(FD3D12Texture* Texture)
 		RenderTargetViews[Index] = Texture->RenderTargetViews[Index];
 	}
 }
+#endif // D3D12RHI_SUPPORTS_ALIASED_RHI_RESOURCES
 
 void FD3D12Texture::ReuseStagingBuffer(TUniquePtr<FD3D12LockedResource>&& LockedResource, uint32 Subresource)
 {
@@ -2973,6 +2975,7 @@ FTextureRHIRef FD3D12DynamicRHI::RHICreateTextureCubeFromResource(EPixelFormat F
 	return CreateTextureFromResource(false, true, Format, TexCreateFlags, ClearValueBinding, Resource);
 }
 
+#if D3D12RHI_SUPPORTS_ALIASED_RHI_RESOURCES
 void FD3D12DynamicRHI::RHIAliasTextureResources(FTextureRHIRef& DestTextureRHI, FTextureRHIRef& SrcTextureRHI)
 {
 	FD3D12Texture* DestTexture = GetD3D12TextureFromRHITexture(DestTextureRHI);
@@ -3048,7 +3051,7 @@ FTextureRHIRef FD3D12DynamicRHI::RHICreateAliasedTexture(FTextureRHIRef& SourceT
 
 	return ReturnTexture;
 }
-
+#endif // D3D12RHI_SUPPORTS_ALIASED_RHI_RESOURCES
 
 ///////////////////////////////////////////////////////////////////////////////////////////
 // FD3D12CommandContext Texture functions
diff --git a/Engine/Source/Runtime/D3D12RHI/Private/D3D12Texture.h b/Engine/Source/Runtime/D3D12RHI/Private/D3D12Texture.h
--- a/Engine/Source/Runtime/D3D12RHI/Private/D3D12Texture.h
+++ b/Engine/Source/Runtime/D3D12RHI/Private/D3D12Texture.h
@@ -67,8 +67,6 @@ public:
 	bool IsStreamable() const									{ return EnumHasAnyFlags(GetDesc().Flags, ETextureCreateFlags::Streamable); }
 	bool SkipsFastClearFinalize() const                         { return EnumHasAnyFlags(GetDesc().Flags, ETextureCreateFlags::NoFastClearFinalize); }
 
-	const FTextureRHIRef& GetAliasingSourceTexture() const		{ return AliasingSourceTexture; }
-
 	void GetReadBackHeapDesc(D3D12_PLACED_SUBRESOURCE_FOOTPRINT& OutFootprint, uint32 Subresource) const;
 
 	bool HasCreatedRTVsPerSlice() const
@@ -142,9 +140,14 @@ public:
 
 	void CopyTextureRegion(FD3D12CommandContext& Context, uint32 DestX, uint32 DestY, uint32 DestZ, FD3D12Texture* SourceTexture, const D3D12_BOX& SourceBox);
 
+#if D3D12RHI_SUPPORTS_ALIASED_RHI_RESOURCES
 	// Resource aliasing
 	void AliasResources(FD3D12Texture* Texture);
-	void SetAliasingSource(FTextureRHIRef& SourceTextureRHI)	{ AliasingSourceTexture = SourceTextureRHI; }
+	void SetAliasingSource(FTextureRHIRef& SourceTextureRHI)
+	{
+		AliasingSourceTexture = SourceTextureRHI;
+	}
+#endif
 
 	// Staging buffer reuse
 	void ReuseStagingBuffer(TUniquePtr<FD3D12LockedResource>&& LockedResource, uint32 Subresource);
@@ -175,8 +178,10 @@ protected:
 	// Cached footprint size of first resource - optimization
 	mutable TUniquePtr<D3D12_PLACED_SUBRESOURCE_FOOTPRINT> FirstSubresourceFootprint;
 
+#if D3D12RHI_SUPPORTS_ALIASED_RHI_RESOURCES
 	// Source texture when aliased
 	FTextureRHIRef AliasingSourceTexture;
+#endif
 };
 
 inline FD3D12RenderTargetView* FD3D12Texture::GetRenderTargetView(int32 MipIndex, int32 ArraySliceIndex) const
diff --git a/Engine/Source/Runtime/D3D12RHI/Private/Windows/WindowsD3D12RHIDefinitions.h b/Engine/Source/Runtime/D3D12RHI/Private/Windows/WindowsD3D12RHIDefinitions.h
--- a/Engine/Source/Runtime/D3D12RHI/Private/Windows/WindowsD3D12RHIDefinitions.h
+++ b/Engine/Source/Runtime/D3D12RHI/Private/Windows/WindowsD3D12RHIDefinitions.h
@@ -36,3 +36,9 @@ typedef ID3D12GraphicsCommandList ID3D12CopyCommandList;
 
 #define D3D12RHI_ALLOW_ENHANCED_BARRIERS 1
 #define D3D12RHI_ALLOW_LEGACY_BARRIERS   1
+
+// Needed for XR
+#define D3D12RHI_SUPPORTS_ALIASED_RHI_RESOURCES 1
+
+#define USE_BUFFER_POOL_ALLOCATOR 1
+#define USE_TEXTURE_POOL_ALLOCATOR 1
