[HWRT]
- Add more profile tags for expensive parts
- Only recreate bindless SRVs when buffer, size or offset changes
- Don't call GetRaytracingAccelerationStructurePrebuildInfo when checks are not enabled

diff --git a/Engine/Source/Runtime/D3D12RHI/Private/D3D12RayTracing.cpp b/Engine/Source/Runtime/D3D12RHI/Private/D3D12RayTracing.cpp
--- a/Engine/Source/Runtime/D3D12RHI/Private/D3D12RayTracing.cpp
+++ b/Engine/Source/Runtime/D3D12RHI/Private/D3D12RayTracing.cpp
@@ -3872,13 +3872,17 @@ FD3D12RayTracingGeometry::~FD3D12RayTracingGeometry()
 	ReleaseUnderlyingResource();
 }
 
-void FD3D12RayTracingGeometry::AllocateBufferSRVs(uint32 InGPUIndex)
+bool FD3D12RayTracingGeometry::AllocateBufferSRVs(uint32 InGPUIndex)
 {
-	HitGroupSystemIndexBufferSRV[InGPUIndex].Reset();
-	HitGroupSystemSegmentVertexBufferSRVs[InGPUIndex].Empty();
+	bool bAnyBuffersUpdated = false;
 
 	// Procedural doesn't need any SRVs for index buffer
-	if (Initializer.IndexBuffer && Initializer.GeometryType == RTGT_Triangles)
+	if (Initializer.IndexBuffer && Initializer.GeometryType == RTGT_Triangles
+#if !DO_CHECK
+		// if checks are not enabled then simply early out when HitGroupSystemIndexBufferSRV is already created because we don't support updating them
+		&& HitGroupSystemIndexBufferSRV[InGPUIndex] == nullptr		
+#endif	
+		)
 	{
 		checkf((Initializer.IndexBufferOffset % RHI_RAW_VIEW_ALIGNMENT) == 0, TEXT("The byte offset of raw views must be a multiple of %d (specified offset: %d)."), RHI_RAW_VIEW_ALIGNMENT, Initializer.IndexBufferOffset);
 
@@ -3902,12 +3906,33 @@ void FD3D12RayTracingGeometry::AllocateBufferSRVs(uint32 InGPUIndex)
 		SRVDesc.Buffer.Flags = D3D12_BUFFER_SRV_FLAG_RAW;
 		SRVDesc.Buffer.StructureByteStride = 0;
 
-		HitGroupSystemIndexBufferSRV[InGPUIndex] = MakeShared<FD3D12ShaderResourceView>(GetParentAdapter()->GetDevice(InGPUIndex), InGPUIndex > 0 ? HitGroupSystemIndexBufferSRV[0].Get() : nullptr, ERHIDescriptorType::BufferSRV);
-		HitGroupSystemIndexBufferSRV[InGPUIndex]->CreateView(IndexBuffer, SRVDesc, FD3D12ShaderResourceView::EFlags::None);
+#if DO_CHECK
+		if (HitGroupSystemIndexBufferSRV[InGPUIndex])
+		{
+			// Validate that current index buffer view hasn't changed
+			D3D12_SHADER_RESOURCE_VIEW_DESC CurrentSRVDesc = HitGroupSystemIndexBufferSRV[InGPUIndex]->GetD3DDesc();
+			check(SRVDesc.Buffer.FirstElement == CurrentSRVDesc.Buffer.FirstElement &&
+				SRVDesc.Buffer.NumElements == CurrentSRVDesc.Buffer.NumElements &&
+				IndexBuffer == HitGroupSystemIndexBufferSRV[InGPUIndex]->GetShaderResource());
+		}
+		else
+#endif // DO_CHECK
+		{
+			HitGroupSystemIndexBufferSRV[InGPUIndex] = MakeShared<FD3D12ShaderResourceView>(GetParentAdapter()->GetDevice(InGPUIndex), InGPUIndex > 0 ? HitGroupSystemIndexBufferSRV[0].Get() : nullptr, ERHIDescriptorType::BufferSRV);
+			HitGroupSystemIndexBufferSRV[InGPUIndex]->CreateView(IndexBuffer, SRVDesc, FD3D12ShaderResourceView::EFlags::None);
+		
+			bAnyBuffersUpdated = true;
+		}
 	}
 
-	for (const FRayTracingGeometrySegment& Segment : Initializer.Segments)
+	if (HitGroupSystemSegmentVertexBufferSRVs[InGPUIndex].Num() != Initializer.Segments.Num())
+	{		
+		HitGroupSystemSegmentVertexBufferSRVs[InGPUIndex].SetNum(Initializer.Segments.Num());
+	}
+
+	for (int32 SegmentIndex = 0; SegmentIndex < Initializer.Segments.Num(); ++SegmentIndex)
 	{
+		const FRayTracingGeometrySegment& Segment = Initializer.Segments[SegmentIndex];
 		checkf((Segment.VertexBufferOffset % RHI_RAW_VIEW_ALIGNMENT) == 0, TEXT("The byte offset of raw views must be a multiple of %d (specified offset: %d)."), RHI_RAW_VIEW_ALIGNMENT, Segment.VertexBufferOffset);
 
 		FD3D12Buffer* VertexBuffer = FD3D12DynamicRHI::ResourceCast(Segment.VertexBuffer.GetReference());
@@ -3928,19 +3953,38 @@ void FD3D12RayTracingGeometry::AllocateBufferSRVs(uint32 InGPUIndex)
 		SRVDesc.Buffer.Flags = D3D12_BUFFER_SRV_FLAG_RAW;
 		SRVDesc.Buffer.StructureByteStride = 0;
 
-		FD3D12ShaderResourceView* FirstLinkedObject = nullptr;
-		if (InGPUIndex > 0)
+		bool bCreateNewSRV = true;
+		FD3D12ShaderResourceView* CurrentVertexBufferSRV = SegmentIndex < HitGroupSystemSegmentVertexBufferSRVs[InGPUIndex].Num() ? HitGroupSystemSegmentVertexBufferSRVs[InGPUIndex][SegmentIndex].Get() : nullptr;
+		if (CurrentVertexBufferSRV)
 		{
-			int32 SegmentIndex = HitGroupSystemSegmentVertexBufferSRVs[InGPUIndex].Num();
-			if (HitGroupSystemSegmentVertexBufferSRVs[0].Num() > SegmentIndex)
+			D3D12_SHADER_RESOURCE_VIEW_DESC CurrentSRVDesc = CurrentVertexBufferSRV->GetD3DDesc();
+			if (SRVDesc.Buffer.FirstElement == CurrentSRVDesc.Buffer.FirstElement &&
+				SRVDesc.Buffer.NumElements == CurrentSRVDesc.Buffer.NumElements &&
+				CurrentVertexBufferSRV->GetShaderResource() == VertexBuffer)
 			{
-				FirstLinkedObject = HitGroupSystemSegmentVertexBufferSRVs[0][SegmentIndex].Get();
+				bCreateNewSRV = false;
 			}
 		}
-		TSharedPtr<FD3D12ShaderResourceView> VertexBufferSRV = MakeShared<FD3D12ShaderResourceView>(GetParentAdapter()->GetDevice(InGPUIndex), FirstLinkedObject, ERHIDescriptorType::BufferSRV);
-		VertexBufferSRV->CreateView(VertexBuffer, SRVDesc, FD3D12ShaderResourceView::EFlags::None);
-		HitGroupSystemSegmentVertexBufferSRVs[InGPUIndex].Add(VertexBufferSRV);
+
+		if (bCreateNewSRV)
+		{
+			FD3D12ShaderResourceView* FirstLinkedObject = nullptr;
+			if (InGPUIndex > 0)
+			{
+				if (HitGroupSystemSegmentVertexBufferSRVs[0].Num() > SegmentIndex)
+				{
+					FirstLinkedObject = HitGroupSystemSegmentVertexBufferSRVs[0][SegmentIndex].Get();
+				}
+			}
+			TSharedPtr<FD3D12ShaderResourceView> VertexBufferSRV = MakeShared<FD3D12ShaderResourceView>(GetParentAdapter()->GetDevice(InGPUIndex), FirstLinkedObject, ERHIDescriptorType::BufferSRV);
+			VertexBufferSRV->CreateView(VertexBuffer, SRVDesc, FD3D12ShaderResourceView::EFlags::None);
+			HitGroupSystemSegmentVertexBufferSRVs[InGPUIndex][SegmentIndex] = VertexBufferSRV;
+
+			bAnyBuffersUpdated = true;
+		}
 	}
+
+	return bAnyBuffersUpdated;
 }
 
 void FD3D12RayTracingGeometry::RegisterAsRenameListener(uint32 InGPUIndex)
@@ -4068,52 +4112,56 @@ void FD3D12RayTracingGeometry::SetupHitGroupSystemParameters(uint32 InGPUIndex)
 
 	bool bBindless = AreBindlessResourcesEnabledForRayTracing(GetParentAdapter());
 
-	TArray<FD3D12HitGroupSystemParameters>& HitGroupSystemParametersForThisGPU = HitGroupSystemParameters[InGPUIndex];
-	HitGroupSystemParametersForThisGPU.Reset(Initializer.Segments.Num());
-
+	bool bAnyBuffersUpdated = true;
 	check(BuffersValid(InGPUIndex));
 	if (bBindless)
 	{
-		AllocateBufferSRVs(InGPUIndex);
+		bAnyBuffersUpdated = AllocateBufferSRVs(InGPUIndex);
 	}
 
-	FD3D12Buffer* IndexBuffer = FD3D12DynamicRHI::ResourceCast(Initializer.IndexBuffer.GetReference(), InGPUIndex);
-	const uint32 IndexStride = IndexBuffer ? IndexBuffer->GetStride() : 0;
-	for (int32 SegmentIndex = 0; SegmentIndex < Initializer.Segments.Num(); ++SegmentIndex)
+	if (bAnyBuffersUpdated)
 	{
-		const FRayTracingGeometrySegment& Segment = Initializer.Segments[SegmentIndex];
-		FD3D12Buffer* VertexBuffer = FD3D12DynamicRHI::ResourceCast(Segment.VertexBuffer.GetReference(), InGPUIndex);
+		TArray<FD3D12HitGroupSystemParameters>& HitGroupSystemParametersForThisGPU = HitGroupSystemParameters[InGPUIndex];
+		HitGroupSystemParametersForThisGPU.Reset(Initializer.Segments.Num());
 
-		FD3D12HitGroupSystemParameters SystemParameters = {};
-		SystemParameters.RootConstants.SetVertexAndIndexStride(Segment.VertexBufferStride, IndexStride);
-		if (bBindless)
+		FD3D12Buffer* IndexBuffer = FD3D12DynamicRHI::ResourceCast(Initializer.IndexBuffer.GetReference(), InGPUIndex);
+		const uint32 IndexStride = IndexBuffer ? IndexBuffer->GetStride() : 0;
+		for (int32 SegmentIndex = 0; SegmentIndex < Initializer.Segments.Num(); ++SegmentIndex)
 		{
-			SystemParameters.BindlessHitGroupSystemVertexBuffer = HitGroupSystemSegmentVertexBufferSRVs[InGPUIndex][SegmentIndex]->GetBindlessHandle().GetIndex();
-		}
-		else
-		{
-			SystemParameters.VertexBuffer = VertexBuffer->ResourceLocation.GetGPUVirtualAddress() + Segment.VertexBufferOffset;
-		}
+			const FRayTracingGeometrySegment& Segment = Initializer.Segments[SegmentIndex];
+			FD3D12Buffer* VertexBuffer = FD3D12DynamicRHI::ResourceCast(Segment.VertexBuffer.GetReference(), InGPUIndex);
 
-		if (GeometryType == D3D12_RAYTRACING_GEOMETRY_TYPE_TRIANGLES && IndexBuffer != nullptr)
-		{
+			FD3D12HitGroupSystemParameters SystemParameters = {};
+			SystemParameters.RootConstants.SetVertexAndIndexStride(Segment.VertexBufferStride, IndexStride);
 			if (bBindless)
 			{
-				SystemParameters.BindlessHitGroupSystemIndexBuffer = HitGroupSystemIndexBufferSRV[InGPUIndex]->GetBindlessHandle().GetIndex();
+				SystemParameters.BindlessHitGroupSystemVertexBuffer = HitGroupSystemSegmentVertexBufferSRVs[InGPUIndex][SegmentIndex]->GetBindlessHandle().GetIndex();
 			}
 			else
 			{
-				SystemParameters.IndexBuffer = IndexBuffer->ResourceLocation.GetGPUVirtualAddress();
+				SystemParameters.VertexBuffer = VertexBuffer->ResourceLocation.GetGPUVirtualAddress() + Segment.VertexBufferOffset;
 			}
-			SystemParameters.RootConstants.IndexBufferOffsetInBytes = Initializer.IndexBufferOffset + IndexStride * Segment.FirstPrimitive * FD3D12RayTracingGeometry::IndicesPerPrimitive;
-			SystemParameters.RootConstants.FirstPrimitive = Segment.FirstPrimitive;
-		}
 
-		HitGroupSystemParametersForThisGPU.Add(SystemParameters);
-	}
+			if (GeometryType == D3D12_RAYTRACING_GEOMETRY_TYPE_TRIANGLES && IndexBuffer != nullptr)
+			{
+				if (bBindless)
+				{
+					SystemParameters.BindlessHitGroupSystemIndexBuffer = HitGroupSystemIndexBufferSRV[InGPUIndex]->GetBindlessHandle().GetIndex();
+				}
+				else
+				{
+					SystemParameters.IndexBuffer = IndexBuffer->ResourceLocation.GetGPUVirtualAddress();
+				}
+				SystemParameters.RootConstants.IndexBufferOffsetInBytes = Initializer.IndexBufferOffset + IndexStride * Segment.FirstPrimitive * FD3D12RayTracingGeometry::IndicesPerPrimitive;
+				SystemParameters.RootConstants.FirstPrimitive = Segment.FirstPrimitive;
+			}
 
-	// Notify listeners about changes
-	HitGroupParamatersUpdated();
+			HitGroupSystemParametersForThisGPU.Add(SystemParameters);
+		}
+		
+		// Notify listeners about changes
+		HitGroupParamatersUpdated();
+	}
 }
 
 void FD3D12RayTracingGeometry::CreateAccelerationStructureBuildDesc(FD3D12CommandContext& CommandContext, EAccelerationStructureBuildMode BuildMode, D3D12_GPU_VIRTUAL_ADDRESS ScratchBufferAddress, D3D12_BUILD_RAYTRACING_ACCELERATION_STRUCTURE_DESC& OutDesc, TArrayView<D3D12_RAYTRACING_GEOMETRY_DESC>& OutGeometryDescs) const
@@ -4235,6 +4283,7 @@ void FD3D12RayTracingGeometry::CreateAccelerationStructureBuildDesc(FD3D12Comman
 	PrebuildDescInputs.pGeometryDescs = OutGeometryDescs.GetData();
 	PrebuildDescInputs.Flags = LocalBuildFlags;
 
+#if DO_CHECK
 	D3D12_RAYTRACING_ACCELERATION_STRUCTURE_PREBUILD_INFO PrebuildInfo = {};
 
 	CommandContext.GetParentDevice()->GetRaytracingAccelerationStructurePrebuildInfo(&PrebuildDescInputs, &PrebuildInfo);
@@ -4250,6 +4299,7 @@ void FD3D12RayTracingGeometry::CreateAccelerationStructureBuildDesc(FD3D12Comman
 	{
 		check(PrebuildInfo.ScratchDataSizeInBytes <= SizeInfo.BuildScratchSize);
 	}
+#endif // DO_CHECK
 
 	D3D12_BUILD_RAYTRACING_ACCELERATION_STRUCTURE_DESC BuildDesc = {};
 	BuildDesc.Inputs = PrebuildDescInputs;
@@ -4464,6 +4514,7 @@ void PrepareAccelerationStructureBuild(
 
 	{
 		// at the same time also gather resources that need to be resident when using TLAS
+		TRACE_CPUPROFILER_EVENT_SCOPE(UpdateTLASResidency);
 
 		TArray<const FD3D12Resource*>& ResourcesToMakeResidentForThisGPU = Scene.ResourcesToMakeResident[GPUIndex];
 
@@ -4613,6 +4664,7 @@ void FD3D12RayTracingScene::UpdateResidency(FD3D12CommandContext& CommandContext
 
 void FD3D12CommandContext::BuildAccelerationStructuresInternal(TConstArrayView<D3D12_BUILD_RAYTRACING_ACCELERATION_STRUCTURE_DESC> BuildDescs)
 {
+	TRACE_CPUPROFILER_EVENT_SCOPE(D3D12::BuildRaytracingAccelerationStructure);
 	for (const D3D12_BUILD_RAYTRACING_ACCELERATION_STRUCTURE_DESC& Desc : BuildDescs)
 	{
 		GraphicsCommandList4()->BuildRaytracingAccelerationStructure(&Desc, 0, nullptr);
diff --git a/Engine/Source/Runtime/D3D12RHI/Private/D3D12RayTracing.h b/Engine/Source/Runtime/D3D12RHI/Private/D3D12RayTracing.h
--- a/Engine/Source/Runtime/D3D12RHI/Private/D3D12RayTracing.h
+++ b/Engine/Source/Runtime/D3D12RHI/Private/D3D12RayTracing.h
@@ -147,7 +147,7 @@ public:
 	// Implement FD3D12ShaderResourceRenameListener interface
 	virtual void ResourceRenamed(FD3D12ContextArray const& Contexts, FD3D12BaseShaderResource* InRenamedResource, FD3D12ResourceLocation* InNewResourceLocation) override;
 
-	void AllocateBufferSRVs(uint32 InGPUIndex);
+	bool AllocateBufferSRVs(uint32 InGPUIndex);
 
 	void RegisterAsRenameListener(uint32 InGPUIndex);
 	void UnregisterAsRenameListener(uint32 InGPUIndex);
