Fixing compilation issues with the new Predication code.

diff --git a/Engine/Source/Runtime/D3D12RHI/Private/D3D12CommandContext.h b/Engine/Source/Runtime/D3D12RHI/Private/D3D12CommandContext.h
--- a/Engine/Source/Runtime/D3D12RHI/Private/D3D12CommandContext.h
+++ b/Engine/Source/Runtime/D3D12RHI/Private/D3D12CommandContext.h
@@ -701,9 +701,9 @@ public:
 		FRHIBuffer* InPredicationBuffer,
 		ERHIPredicationOp InPredicationOp,
 		uint32 InOffset
-	) final override;
+	) override;
 
-	virtual void RHIUnsetPredication() final override;
+	virtual void RHIUnsetPredication() override;
 
 	virtual void RHIDrawPrimitive(uint32 BaseVertexIndex, uint32 NumPrimitives, uint32 NumInstances) final override;
 	virtual void RHIDrawPrimitiveIndirect(FRHIBuffer* ArgumentBuffer, uint32 ArgumentOffset) final override;
