- Allow DX12 Present to be called on the Submission Thread (r.D3D12.PresentOnSubmissionThread)
- Refactor Present calls to use FRHIPresentArg in order have the right game frame number correlating with the Present counter
- Add more instrumentation around Present calls and sync adding frame information
- Add a -1 option for r.FinishCurrentFrame to skip the end frame sync

diff --git a/Engine/Plugins/Runtime/OpenXR/Source/OpenXRHMD/Internal/OpenXRHMD_RenderBridge.h b/Engine/Plugins/Runtime/OpenXR/Source/OpenXRHMD/Internal/OpenXRHMD_RenderBridge.h
--- a/Engine/Plugins/Runtime/OpenXR/Source/OpenXRHMD/Internal/OpenXRHMD_RenderBridge.h
+++ b/Engine/Plugins/Runtime/OpenXR/Source/OpenXRHMD/Internal/OpenXRHMD_RenderBridge.h
@@ -63,6 +63,7 @@ public:
 
 	/** FRHICustomPresent */
 	virtual bool Present(IRHICommandContext& RHICmdContext, int32& InOutSyncInterval) override;
+	virtual bool PresentOnSubmissionThread(int32& InOutSyncInterval) override;
 
 	virtual bool Support10BitSwapchain() const { return false; }
 
@@ -71,6 +72,7 @@ public:
 	virtual void SetSkipRate(uint32 SkipRate) {}
 
 	virtual void HMDOnFinishRendering_RHIThread(IRHICommandContext& RHICmdContext);
+	virtual void HMDOnFinishRendering_SubmissionThread();
 
 protected:
 	XrInstance Instance;
diff --git a/Engine/Plugins/Runtime/OpenXR/Source/OpenXRHMD/Private/OpenXRHMD.cpp b/Engine/Plugins/Runtime/OpenXR/Source/OpenXRHMD/Private/OpenXRHMD.cpp
--- a/Engine/Plugins/Runtime/OpenXR/Source/OpenXRHMD/Private/OpenXRHMD.cpp
+++ b/Engine/Plugins/Runtime/OpenXR/Source/OpenXRHMD/Private/OpenXRHMD.cpp
@@ -4263,10 +4263,8 @@ void FOpenXRHMD::OnBeginRendering_RHIThread(IRHICommandContext& RHICmdContext, c
 	}
 }
 
-void FOpenXRHMD::OnFinishRendering_RHIThread(IRHICommandContext& RHICmdContext)
+void FOpenXRHMD::OnFinishRendering(IRHICommandContext* RHICmdContext)
 {
-	check(IsInRenderingThread() || IsInRHIThread());
-
 	SCOPED_NAMED_EVENT(EndFrame, FColor::Red);
 
 	if (!bIsRendering || !RenderBridge)
@@ -4274,12 +4272,12 @@ void FOpenXRHMD::OnFinishRendering_RHIThread(IRHICommandContext& RHICmdContext)
 		return;
 	}
 	
-	UE_LOG(LogHMD, VeryVerbose, TEXT("%s WF_%i FOpenXRHMD::OnFinishRendering_RHIThread releasing swapchain images now."), HMDThreadString(), PipelinedFrameStateRHI.WaitCount, PipelinedFrameStateRHI.BeginCount);
+	UE_LOG(LogHMD, VeryVerbose, TEXT("%s WF_%i FOpenXRHMD::OnFinishRendering releasing swapchain images now."), HMDThreadString(), PipelinedFrameStateRHI.WaitCount, PipelinedFrameStateRHI.BeginCount);
 
 	// We need to ensure we release the swap chain images even if the session is not running.
 	if (PipelinedLayerStateRHI.ColorSwapchain)
 	{
-		IRHICommandContext* const RHICommandContextIfRequired = RuntimeRequiresRHIContext() ? &RHICmdContext : nullptr;
+		IRHICommandContext* const RHICommandContextIfRequired = RuntimeRequiresRHIContext() ? RHICmdContext : nullptr;
 		PipelinedLayerStateRHI.ColorSwapchain->ReleaseCurrentImage_RHIThread(RHICommandContextIfRequired);
 
 		if (PipelinedLayerStateRHI.DepthSwapchain)
@@ -4384,7 +4382,7 @@ void FOpenXRHMD::OnFinishRendering_RHIThread(IRHICommandContext& RHICmdContext)
 		XrRHIContextEPIC RHIContextEPIC = { (XrStructureType)XR_TYPE_RHI_CONTEXT_EPIC };
 		if (RuntimeRequiresRHIContext())
 		{
-			RHIContextEPIC.RHIContext = &RHICmdContext;
+			RHIContextEPIC.RHIContext = RHICmdContext;
 			RHIContextEPIC.next = Next;
 			Next = &RHIContextEPIC;
 		}
diff --git a/Engine/Plugins/Runtime/OpenXR/Source/OpenXRHMD/Private/OpenXRHMD.h b/Engine/Plugins/Runtime/OpenXR/Source/OpenXRHMD/Private/OpenXRHMD.h
--- a/Engine/Plugins/Runtime/OpenXR/Source/OpenXRHMD/Private/OpenXRHMD.h
+++ b/Engine/Plugins/Runtime/OpenXR/Source/OpenXRHMD/Private/OpenXRHMD.h
@@ -469,7 +469,7 @@ public:
 
 	void OnBeginSimulation_GameThread();
 	void OnBeginRendering_RHIThread(IRHICommandContext& RHICmdContext, const FPipelinedFrameState& InFrameState, FXRSwapChainPtr ColorSwapchain, FXRSwapChainPtr DepthSwapchain, FXRSwapChainPtr MotionVectorSwapchain, FXRSwapChainPtr MotionVectorDepthSwapchain, FXRSwapChainPtr EmulationSwapchain);
-	void OnFinishRendering_RHIThread(IRHICommandContext& RHICmdContext);
+	void OnFinishRendering(IRHICommandContext* RHICmdContext);
 
 	/** IOpenXRHMD */
 	void SetInputModule(IOpenXRInputModule* InInputModule) override
diff --git a/Engine/Plugins/Runtime/OpenXR/Source/OpenXRHMD/Private/OpenXRHMD_RenderBridge.cpp b/Engine/Plugins/Runtime/OpenXR/Source/OpenXRHMD/Private/OpenXRHMD_RenderBridge.cpp
--- a/Engine/Plugins/Runtime/OpenXR/Source/OpenXRHMD/Private/OpenXRHMD_RenderBridge.cpp
+++ b/Engine/Plugins/Runtime/OpenXR/Source/OpenXRHMD/Private/OpenXRHMD_RenderBridge.cpp
@@ -20,11 +20,39 @@ bool FOpenXRRenderBridge::Present(IRHICommandContext& RHICmdContext, int32& InOu
 	return bNeedsNativePresent;
 }
 
+bool FOpenXRRenderBridge::PresentOnSubmissionThread(int32& InOutSyncInterval)
+{
+	bool bNeedsNativePresent = true;
+
+	if (OpenXRHMD)
+	{
+		HMDOnFinishRendering_SubmissionThread();
+		bNeedsNativePresent = !OpenXRHMD->IsStandaloneStereoOnlyDevice();
+	}
+
+	InOutSyncInterval = 0; // VSync off
+
+	return bNeedsNativePresent;
+}
+
 void FOpenXRRenderBridge::HMDOnFinishRendering_RHIThread(IRHICommandContext& RHICmdContext)
 {
 	if (OpenXRHMD)
 	{
-		OpenXRHMD->OnFinishRendering_RHIThread(RHICmdContext);
+		check(IsInRenderingThread() || IsInRHIThread());
+		OpenXRHMD->OnFinishRendering(&RHICmdContext);
+	}
+}
+
+void FOpenXRRenderBridge::HMDOnFinishRendering_SubmissionThread()
+{
+	if (OpenXRHMD)
+	{
+		check(!IsInRenderingThread() && !IsInRHIThread());
+		check(OpenXRHMD->RuntimeRequiresRHIContext() == false);
+		
+		IRHICommandContext* RHIContext = nullptr;
+		OpenXRHMD->OnFinishRendering(RHIContext);
 	}
 }
 
diff --git a/Engine/Plugins/Runtime/nDisplay/Source/DisplayCluster/Private/Render/Presentation/DisplayClusterPresentationBase.cpp b/Engine/Plugins/Runtime/nDisplay/Source/DisplayCluster/Private/Render/Presentation/DisplayClusterPresentationBase.cpp
--- a/Engine/Plugins/Runtime/nDisplay/Source/DisplayCluster/Private/Render/Presentation/DisplayClusterPresentationBase.cpp
+++ b/Engine/Plugins/Runtime/nDisplay/Source/DisplayCluster/Private/Render/Presentation/DisplayClusterPresentationBase.cpp
@@ -42,6 +42,16 @@ void FDisplayClusterPresentationBase::OnBackBufferResize()
 }
 
 bool FDisplayClusterPresentationBase::Present(IRHICommandContext& RHICmdContext, int32& InOutSyncInterval)
+{
+	return PresentInternal(InOutSyncInterval);
+}
+
+bool FDisplayClusterPresentationBase::PresentOnSubmissionThread(int32& InOutSyncInterval)
+{
+	return PresentInternal(InOutSyncInterval);
+}
+
+bool FDisplayClusterPresentationBase::PresentInternal(int32& InOutSyncInterval)
 {
 	TRACE_CPUPROFILER_EVENT_SCOPE(nDisplay PresentationBase);
 
@@ -52,7 +62,7 @@ bool FDisplayClusterPresentationBase::Present(IRHICommandContext& RHICmdContext,
 	{
 		// Update sync value with nDisplay value
 		InOutSyncInterval = GetSwapInt();
-		
+
 		GDisplayCluster->GetCallbacks().OnDisplayClusterPresentationPreSynchronization_RHIThread().Broadcast();
 
 		// False results means we don't need to present current frame, the sync object already presented it
diff --git a/Engine/Plugins/Runtime/nDisplay/Source/DisplayCluster/Private/Render/Presentation/DisplayClusterPresentationBase.h b/Engine/Plugins/Runtime/nDisplay/Source/DisplayCluster/Private/Render/Presentation/DisplayClusterPresentationBase.h
--- a/Engine/Plugins/Runtime/nDisplay/Source/DisplayCluster/Private/Render/Presentation/DisplayClusterPresentationBase.h
+++ b/Engine/Plugins/Runtime/nDisplay/Source/DisplayCluster/Private/Render/Presentation/DisplayClusterPresentationBase.h
@@ -33,10 +33,14 @@ protected:
 	{ return true; };
 
 	virtual bool Present(IRHICommandContext& RHICmdContext, int32& InOutSyncInterval) override;
+	virtual bool PresentOnSubmissionThread(int32& InOutSyncInterval) override;
 	virtual void PostPresent() override;
 	virtual void OnBackBufferResize() override;
 
 protected:
+
+	bool PresentInternal(int32& InOutSyncInterval);
+
 	FViewport* const GetViewport() const
 	{ return Viewport; }
 
diff --git a/Engine/Plugins/Tests/RHITests/Source/RHITests/Private/RHIClearTextureTests.cpp b/Engine/Plugins/Tests/RHITests/Source/RHITests/Private/RHIClearTextureTests.cpp
--- a/Engine/Plugins/Tests/RHITests/Source/RHITests/Private/RHIClearTextureTests.cpp
+++ b/Engine/Plugins/Tests/RHITests/Source/RHITests/Private/RHIClearTextureTests.cpp
@@ -47,7 +47,7 @@ static void Test_ClearTextureBeginCapture(FRHICommandListImmediate& RHICmdList)
 {
 	if (GBeginEndCaptureHack)
 	{
-		RHICmdList.EndDrawingViewport(nullptr, true, false);
+		RHICmdList.EndDrawingViewport(nullptr, FRHIPresentArgs(/*.FrameCounter*/ GFrameCounterRenderThread, /*.bPresent*/ true, /*.bLockToVsync*/ false));
 		GDynamicRHI->RHIWaitForFlip(100000);
 		IRenderCaptureProvider::Get().CaptureFrame();
  		GDynamicRHI->RHISignalFlipEvent();
@@ -64,7 +64,7 @@ static void Test_ClearTextureEndCapture(FRHICommandListImmediate& RHICmdList)
 {
 	if (GBeginEndCaptureHack)
 	{
-		RHICmdList.EndDrawingViewport(nullptr, true, false);
+		RHICmdList.EndDrawingViewport(nullptr, FRHIPresentArgs(/*.FrameCounter*/ GFrameCounterRenderThread, /*.bPresent*/ true, /*.bLockToVsync*/ false));
 		GDynamicRHI->RHIWaitForFlip(100000);
 	}
 	else
diff --git a/Engine/Source/Runtime/Apple/MetalRHI/Private/MetalDevice.cpp b/Engine/Source/Runtime/Apple/MetalRHI/Private/MetalDevice.cpp
--- a/Engine/Source/Runtime/Apple/MetalRHI/Private/MetalDevice.cpp
+++ b/Engine/Source/Runtime/Apple/MetalRHI/Private/MetalDevice.cpp
@@ -566,7 +566,7 @@ void FMetalDevice::EnumerateFeatureSupport()
 #endif
 }
 
-void FMetalDevice::EndDrawingViewport(bool bPresent)
+void FMetalDevice::EndDrawingViewport(const FRHIPresentArgs& InPresentArgs)
 {
 	// We may be limiting our framerate to the display link
 	if( FrameReadyEvent != nullptr && !GMetalSeparatePresentThread )
@@ -575,7 +575,7 @@ void FMetalDevice::EndDrawingViewport(bool bPresent)
 		FrameReadyEvent->Wait(MAX_uint32, bIgnoreThreadIdleStats);
 	}
 	
-	if(bPresent)
+	if(InPresentArgs.bPresent)
 	{
 		CaptureManager->PresentFrame(FrameCounter++);
 	}
diff --git a/Engine/Source/Runtime/Apple/MetalRHI/Private/MetalDevice.h b/Engine/Source/Runtime/Apple/MetalRHI/Private/MetalDevice.h
--- a/Engine/Source/Runtime/Apple/MetalRHI/Private/MetalDevice.h
+++ b/Engine/Source/Runtime/Apple/MetalRHI/Private/MetalDevice.h
@@ -109,7 +109,7 @@ public:
 	
 	inline FMetalResourceHeap& GetResourceHeap(void) { return Heap; }
 	
-	void EndDrawingViewport(bool bPresent);
+	void EndDrawingViewport(const FRHIPresentArgs& InPresentArgs);
 	
 	MTLTexturePtr CreateTexture(FMetalSurface* Surface, MTL::TextureDescriptor* Descriptor);
 	FMetalBufferPtr CreatePooledBuffer(FMetalPooledBufferArgs const& Args);
diff --git a/Engine/Source/Runtime/Apple/MetalRHI/Private/MetalViewport.cpp b/Engine/Source/Runtime/Apple/MetalRHI/Private/MetalViewport.cpp
--- a/Engine/Source/Runtime/Apple/MetalRHI/Private/MetalViewport.cpp
+++ b/Engine/Source/Runtime/Apple/MetalRHI/Private/MetalViewport.cpp
@@ -888,7 +888,7 @@ void FMetalDynamicRHI::RHITick( float DeltaTime )
  *	Viewport functions.
  *=============================================================================*/
 
-void FMetalRHICommandContext::RHIEndDrawingViewport(FRHIViewport* ViewportRHI,bool bPresent,bool bLockToVsync)
+void FMetalRHICommandContext::RHIEndDrawingViewport(FRHIViewport* ViewportRHI, const FRHIPresentArgs& InPresentArgs)
 {
     MTL_SCOPED_AUTORELEASE_POOL;
     
@@ -896,7 +896,7 @@ void FMetalRHICommandContext::RHIEndDrawingViewport(FRHIViewport* ViewportRHI,bo
 
 	// enqueue a present if desired
 	static bool const bOffscreenOnly = FParse::Param(FCommandLine::Get(), TEXT("MetalOffscreenOnly"));
-	if (bPresent && !bOffscreenOnly)
+	if (InPresentArgs.bPresent && !bOffscreenOnly)
 	{
 		bool bNeedNativePresent = true;
 #if PLATFORM_MAC || PLATFORM_VISIONOS
@@ -929,11 +929,11 @@ void FMetalRHICommandContext::RHIEndDrawingViewport(FRHIViewport* ViewportRHI,bo
 		
 		if (bNeedNativePresent)
 		{
-			Viewport->Present(CommandQueue, bLockToVsync);
+			Viewport->Present(CommandQueue, InPresentArgs.bLockToVsync);
 		}
 	}
 	
-	Device.EndDrawingViewport(bPresent);
+	Device.EndDrawingViewport(InPresentArgs);
 	
 	Viewport->ReleaseDrawable();
 }
diff --git a/Engine/Source/Runtime/Apple/MetalRHI/Public/MetalRHIContext.h b/Engine/Source/Runtime/Apple/MetalRHI/Public/MetalRHIContext.h
--- a/Engine/Source/Runtime/Apple/MetalRHI/Public/MetalRHIContext.h
+++ b/Engine/Source/Runtime/Apple/MetalRHI/Public/MetalRHIContext.h
@@ -91,7 +91,7 @@ public:
 	virtual void RHIEndRenderQuery(FRHIRenderQuery* RenderQuery) final override;
 	
 	// This method is queued with an RHIThread, otherwise it will flush after it is queued; without an RHI thread there is no benefit to queuing this frame advance commands
-	virtual void RHIEndDrawingViewport(FRHIViewport* Viewport, bool bPresent, bool bLockToVsync) override;
+	virtual void RHIEndDrawingViewport(FRHIViewport* Viewport, const FRHIPresentArgs& InPresentArgs) override;
 	
 	virtual void RHISetStreamSource(uint32 StreamIndex, FRHIBuffer* VertexBuffer, uint32 Offset) final override;
 
diff --git a/Engine/Source/Runtime/Core/Private/HAL/ConsoleManager.cpp b/Engine/Source/Runtime/Core/Private/HAL/ConsoleManager.cpp
--- a/Engine/Source/Runtime/Core/Private/HAL/ConsoleManager.cpp
+++ b/Engine/Source/Runtime/Core/Private/HAL/ConsoleManager.cpp
@@ -4192,7 +4192,9 @@ static TAutoConsoleVariable<int32> CVarSetVSyncEditorEnabled(
 static TAutoConsoleVariable<int32> CVarFinishCurrentFrame(
 	TEXT("r.FinishCurrentFrame"),
 	0,
-	TEXT("If on, the current frame will be forced to finish and render to the screen instead of being buffered.  This will improve latency, but slow down overall performance."),
+	TEXT(" 0: Wait for previous frame to be Finished on the GPU.(default)\n"
+		"-1: Do no wait for GPU.\n"
+		" 1: Wait for current frame to be Finished on the GPU."),
 	ECVF_Scalability | ECVF_RenderThreadSafe);
 
 static TAutoConsoleVariable<int32> CVarMaxAnistropy(
diff --git a/Engine/Source/Runtime/D3D12RHI/Private/D3D12CommandContext.cpp b/Engine/Source/Runtime/D3D12RHI/Private/D3D12CommandContext.cpp
--- a/Engine/Source/Runtime/D3D12RHI/Private/D3D12CommandContext.cpp
+++ b/Engine/Source/Runtime/D3D12RHI/Private/D3D12CommandContext.cpp
@@ -600,6 +600,19 @@ void FD3D12ContextCommon::NewPayload()
 	Payloads.Add(new FD3D12Payload(Device->GetQueue(QueueType)));
 }
 
+// Set the Present Params in the context for the Submission Thread to call the native Present
+void FD3D12ContextCommon::SetPresent(FD3D12PresentArgs& PresentArgs)
+{
+	if (IsPendingCommands())
+	{
+		CloseCommandList();
+	}
+
+	FD3D12Payload* PresentPayload = GetPayload(EPhase::Signal);
+	check(PresentPayload->PresentArgs.IsValid() == false);
+	PresentPayload->PresentArgs = PresentArgs;
+}
+
 void FD3D12ContextCommon::FlushCommands(ED3D12FlushFlags FlushFlags)
 {
 #if D3D12RHI_IDLE_AFTER_EVERY_GPU_EVENT
diff --git a/Engine/Source/Runtime/D3D12RHI/Private/D3D12CommandContext.h b/Engine/Source/Runtime/D3D12RHI/Private/D3D12CommandContext.h
--- a/Engine/Source/Runtime/D3D12RHI/Private/D3D12CommandContext.h
+++ b/Engine/Source/Runtime/D3D12RHI/Private/D3D12CommandContext.h
@@ -363,6 +363,9 @@ public:
 		}
 	}
 
+	// Set the Present Params in the context for the Submission Thread to call the native Present
+	void SetPresent(FD3D12PresentArgs& PresentArgs);
+
 	// Flushes any pending commands in this context to the GPU.
 	void FlushCommands(ED3D12FlushFlags FlushFlags = ED3D12FlushFlags::None);
 
@@ -480,7 +483,7 @@ class FD3D12CommandContextBase : public IRHICommandContext, public FD3D12Adapter
 public:
 	FD3D12CommandContextBase(FD3D12Adapter* InParent, FRHIGPUMask InGPUMask);
 
-	void RHIEndDrawingViewport(FRHIViewport* Viewport, bool bPresent, bool bLockToVsync) final override;
+	void RHIEndDrawingViewport(FRHIViewport* Viewport, const FRHIPresentArgs& InPresentArgs) final override;
 
 	FRHIGPUMask GetGPUMask() const { return GPUMask; }
 	FRHIGPUMask GetPhysicalGPUMask() const { return PhysicalGPUMask; }
diff --git a/Engine/Source/Runtime/D3D12RHI/Private/D3D12Submission.cpp b/Engine/Source/Runtime/D3D12RHI/Private/D3D12Submission.cpp
--- a/Engine/Source/Runtime/D3D12RHI/Private/D3D12Submission.cpp
+++ b/Engine/Source/Runtime/D3D12RHI/Private/D3D12Submission.cpp
@@ -879,6 +879,11 @@ void FD3D12DynamicRHI::FlushBatchedPayloads(FD3D12Queue::FPayloadArray& Payloads
 		}
 	#endif
 
+		if (Payload->PresentArgs.IsValid())
+		{
+			Payload->PresentArgs.DrawingViewport->PresentOnSubmissionThread(Payload->PresentArgs);
+		}
+
 		// Submission of this payload is completed. Signal the submission event if one was provided.
 		if (Payload->SubmissionEvent)
 		{
diff --git a/Engine/Source/Runtime/D3D12RHI/Private/D3D12Submission.h b/Engine/Source/Runtime/D3D12RHI/Private/D3D12Submission.h
--- a/Engine/Source/Runtime/D3D12RHI/Private/D3D12Submission.h
+++ b/Engine/Source/Runtime/D3D12RHI/Private/D3D12Submission.h
@@ -22,6 +22,7 @@ class FD3D12Queue;
 class FD3D12Timing;
 class FD3D12Buffer;
 class FD3D12Resource;
+class FD3D12Viewport;
 
 class FD3D12SyncPoint;
 using FD3D12SyncPointRef = TRefCountPtr<FD3D12SyncPoint>;
@@ -189,6 +190,22 @@ struct FD3D12PayloadBaseFixLayout
 	virtual ~FD3D12PayloadBaseFixLayout() = default;
 };
 
+struct FD3D12PresentArgs
+{
+	//Sync event -> Present called
+	FGraphEventRef	PresentEvent;
+
+	// In params
+	FD3D12Viewport*	DrawingViewport = nullptr;
+	uint64 FrameCounter = 0;
+	bool bLockToVsync = false;
+
+	inline bool IsValid() const
+	{
+		return DrawingViewport != nullptr;
+	}
+};
+
 // A single unit of work (specific to a single GPU node and queue type) to be processed by the submission thread.
 struct FD3D12PayloadBase : public FD3D12PayloadBaseFixLayout
 {
@@ -258,6 +275,8 @@ struct FD3D12PayloadBase : public FD3D12PayloadBaseFixLayout
 
 	FD3D12BatchedPayloadObjects BatchedObjects;
 
+	FD3D12PresentArgs PresentArgs;
+
 #if WITH_RHI_BREADCRUMBS
 	FRHIBreadcrumbRange BreadcrumbRange {};
 	TSharedPtr<FRHIBreadcrumbAllocatorArray> BreadcrumbAllocators {};
@@ -294,7 +313,7 @@ struct FD3D12PayloadBase : public FD3D12PayloadBaseFixLayout
 
 	virtual bool HasSignalWork() const
 	{
-		return RequiresQueueFenceSignal() || ManualFencesToSignal.Num() > 0 || SubmissionEvent != nullptr
+		return RequiresQueueFenceSignal() || ManualFencesToSignal.Num() > 0 || SubmissionEvent != nullptr || PresentArgs.IsValid()
 		#if RHI_NEW_GPU_PROFILER
 			|| EndFrameEvent.IsSet()
 		#endif
diff --git a/Engine/Source/Runtime/D3D12RHI/Private/D3D12Viewport.cpp b/Engine/Source/Runtime/D3D12RHI/Private/D3D12Viewport.cpp
--- a/Engine/Source/Runtime/D3D12RHI/Private/D3D12Viewport.cpp
+++ b/Engine/Source/Runtime/D3D12RHI/Private/D3D12Viewport.cpp
@@ -11,6 +11,7 @@
 #include "HDRHelper.h"
 #include "DataDrivenShaderPlatformInfo.h"
 #include "RHIUtilities.h"
+#include "ProfilingDebugging/AssetMetadataTrace.h"
 
 namespace D3D12RHI
 {
@@ -34,6 +35,14 @@
 #endif
 	};
 
+	bool GUseSubmissionThreadPresent = true;
+	static FAutoConsoleVariableRef CVarD3D12SubmissionTimeout(
+		TEXT("r.D3D12.PresentOnSubmissionThread"),
+		GUseSubmissionThreadPresent,
+		TEXT("When true, Present is called on the Submission Thread instead of the RHI Thread"),
+		ECVF_RenderThreadSafe);
+
+
 	int32 GD3D12NumViewportBuffers = D3D12RHI_DEFAULT_NUM_BACKBUFFER;
 	static FAutoConsoleVariableRef CVarD3D12NumViewportBuffers(
 		TEXT("r.D3D12.NumViewportBuffers"),
@@ -227,6 +236,9 @@
 	}
 #endif // D3D12_VIEWPORT_EXPOSES_SWAP_CHAIN
 
+	//Do not destroy FD3D12Viewport if waiting for a Present
+	WaitForLastPresent();	
+
 	GetParentAdapter()->GetViewports().Remove(this);
 
 	FinalDestroyInternal();
@@ -516,116 +528,153 @@
 }
 
 /** Presents the swap chain checking the return result. */
-bool FD3D12Viewport::PresentChecked(IRHICommandContext& RHICmdContext, int32 SyncInterval)
+void FD3D12Viewport::PresentChecked(int32 SyncInterval)
 {
-#if PLATFORM_WINDOWS
-	// We can't call Present if !bIsValid, as it waits a window message to be processed, but the main thread may not be pumping the message handler.
-	if (bIsValid && SwapChain1.IsValid())
+	static constexpr uint32 MaxPresentFailures = 5;
+
+	// Present the back buffer to the viewport window.
+	// In case presentation failures are transient, don't fault on the first one.
+	HRESULT Result = PresentInternal(SyncInterval);
+	if (SUCCEEDED(Result))
 	{
-		// Check if the viewport's swap chain has been invalidated by DXGI.
-		BOOL bSwapChainFullscreenState;
-		TRefCountPtr<IDXGIOutput> SwapChainOutput;
-		SwapChain1->GetFullscreenState(&bSwapChainFullscreenState, SwapChainOutput.GetInitReference());
-		// Can't compare BOOL with bool...
-		if ( (!!bSwapChainFullscreenState)  != bIsFullscreen )
+		CheckedPresentFailureCounter = 0;
+	}
+	else if (!IsTransientPresentationError(Result) || ++CheckedPresentFailureCounter >= MaxPresentFailures)
+	{			
+		VERIFYD3D12RESULT_LAMBDA(Result, GetParentAdapter()->GetD3DDevice(), [this]
 		{
-			bFullscreenLost = true;
-			bIsValid = false;
-		}
+			return GetStateString();
+		});
+	}
+	else
+	{
+		UE_LOG(
+			LogD3D12RHI, Error,
+			TEXT("Swapchain presentation try %u/%u failed with HR(0x%x): %s"),
+			CheckedPresentFailureCounter, MaxPresentFailures,
+			Result, *GetStateString()
+		);
 	}
 
-	if (!bIsValid)
+#if LOG_PRESENT
+	const FString& ThreadName = FThreadManager::GetThreadName(FPlatformTLS::GetCurrentThreadId());
+	UE_LOG(LogD3D12RHI, Log, TEXT("*** PRESENT: Thread %s: Viewport %#016llx: BackBuffer %#016llx (SyncInterval %u) ***"), ThreadName.GetCharArray().GetData(), this, GetBackBuffer_RHIThread(), SyncInterval);
+#endif
+}
+
+UE_TRACE_EVENT_BEGIN(Cpu, GetBackBufferOnRhiThread, NoSync)
+UE_TRACE_EVENT_FIELD(uint8, BackBuffer)
+UE_TRACE_EVENT_END()
+
+FD3D12Texture* FD3D12Viewport::GetBackBuffer_RHIThread(bool bWaitForLastPresent) const
+{
+	TRACE_CPUPROFILER_EVENT_SCOPE(FD3D12Viewport::GetBackBuffer_RHIThread);
+
+#if CPUPROFILERTRACE_ENABLED
+	UE_TRACE_LOG_SCOPED_T(Cpu, GetBackBufferOnRhiThread, CpuChannel)
+		<< GetBackBufferOnRhiThread.BackBuffer(CurrentBackBufferIndex_RHIThread);
+#endif // CPUPROFILERTRACE_ENABLED
+
+	checkSlow(CurrentBackBuffer_RHIThread);
+
+	if (bWaitForLastPresent)
 	{
-#if WITH_MGPU
-		// Present failed so current expected GPU index will not match anymore, so patch up expected back buffer index
-		// Warning: Present is skipped for this frame but could cause a black screen for the next frame as well
-		FScopeLock Lock(&ExpectedBackBufferIndexLock);
-		SetBackBufferIndex_RenderThread((ExpectedBackBufferIndex_RenderThread == 0) ? NumBackBuffers - 1 : ExpectedBackBufferIndex_RenderThread - 1);
-#endif // WITH_MGPU
-		return false;
+		WaitForLastPresent();
 	}
-#endif
+
+	return CurrentBackBuffer_RHIThread->Texture;
+}
+
+
+UE_TRACE_EVENT_BEGIN(Cpu, SchedulePresent, NoSync)
+UE_TRACE_EVENT_FIELD(uint8, BackBuffer)
+UE_TRACE_EVENT_FIELD(uint64, PresentCounter)
+UE_TRACE_EVENT_FIELD(uint64, FrameCounter)
+UE_TRACE_EVENT_END()
+
+UE_TRACE_EVENT_BEGIN(Cpu, Present_SubmissionThread, NoSync)
+UE_TRACE_EVENT_FIELD(uint8, BackBuffer)
+UE_TRACE_EVENT_FIELD(uint64, PresentCounter)
+UE_TRACE_EVENT_FIELD(uint64, FrameCounter)
+UE_TRACE_EVENT_END()
+
+UE_TRACE_EVENT_BEGIN(Cpu, End_Rendering, NoSync)
+UE_TRACE_EVENT_FIELD(uint8, BackBuffer)
+UE_TRACE_EVENT_FIELD(uint64, PresentCounter)
+UE_TRACE_EVENT_FIELD(uint64, FrameCounter)
+UE_TRACE_EVENT_END()
+
+
+bool FD3D12Viewport::SchedulePresent(FD3D12CommandContextBase* BaseContext, const FRHIPresentArgs& InPresentArgs)
+{	
+#if CPUPROFILERTRACE_ENABLED
+	UE_TRACE_LOG_SCOPED_T(Cpu, SchedulePresent, CpuChannel)
+		<< SchedulePresent.BackBuffer(CurrentBackBufferIndex_RHIThread)
+		<< SchedulePresent.PresentCounter(GRHIPresentCounter+1)
+		<< SchedulePresent.FrameCounter(InPresentArgs.FrameCounter)
+		;
+#endif // CPUPROFILERTRACE_ENABLED
+
+	check(BaseContext->GetParentAdapter() == GetParentAdapter());
+	check(InPresentArgs.bPresent);
+
+	WaitForLastPresent();
+
+	FD3D12PresentArgs PresentArgs;
+	PresentEvent = FGraphEvent::CreateGraphEvent();
+	PresentArgs.DrawingViewport = this;
+	PresentArgs.bLockToVsync = InPresentArgs.bLockToVsync;
+	PresentArgs.FrameCounter = InPresentArgs.FrameCounter;
+	PresentArgs.PresentEvent = PresentEvent;
+
+	// MPGU Present only on first GPU context
+	FD3D12CommandContext& Context = *BaseContext->GetSingleDeviceContext(0);
+	Context.SetPresent(PresentArgs);
+
+	return true;
+}
+
+void FD3D12Viewport::PresentOnSubmissionThread(FD3D12PresentArgs& PresentArgs)
+{
+	//Idle scopes needs to know that this is a RHI related thread
+	FTaskTagScope TaskTagScope(ETaskTag::EParallelRhiThread);
+
+#if CPUPROFILERTRACE_ENABLED
+	UE_TRACE_LOG_SCOPED_T(Cpu, Present_SubmissionThread, CpuChannel)
+		<< Present_SubmissionThread.BackBuffer(CurrentBackBufferIndex_RHIThread)
+		<< Present_SubmissionThread.PresentCounter(GRHIPresentCounter+1)
+		<< Present_SubmissionThread.FrameCounter(PresentArgs.FrameCounter)
+		;
+#endif // CPUPROFILERTRACE_ENABLED
+	check(PresentArgs.IsValid());
+
+
+	int32 SyncInterval = PresentArgs.bLockToVsync ? RHIGetSyncInterval() : 0;
 
 	bool bNeedNativePresent = true;
+
 	if (IsValidRef(CustomPresent))
 	{
 		SCOPE_CYCLE_COUNTER(STAT_D3D12CustomPresentTime);
-		bNeedNativePresent = CustomPresent->Present(RHICmdContext, SyncInterval);
+		// CustomPresent::Present can decide to change the SyncInterval to toggle the Vsync
+		bNeedNativePresent = CustomPresent->PresentOnSubmissionThread(SyncInterval);
 	}
-
+	
 	if (bNeedNativePresent)
 	{
-		static constexpr uint32 MaxPresentFailures = 5;
-
-		// Present the back buffer to the viewport window.
-		// In case presentation failures are transient, don't fault on the first one.
-		HRESULT Result = PresentInternal(SyncInterval);
-		if (SUCCEEDED(Result))
-		{
-			CheckedPresentFailureCounter = 0;
-		}
-		else if (!IsTransientPresentationError(Result) || ++CheckedPresentFailureCounter >= MaxPresentFailures)
-		{			
-			VERIFYD3D12RESULT_LAMBDA(Result, GetParentAdapter()->GetD3DDevice(), [this]
-			{
-				return GetStateString();
-			});
-		}
-		else
-		{
-			UE_LOG(
-				LogD3D12RHI, Error,
-				TEXT("Swapchain presentation try %u/%u failed with HR(0x%x): %s"),
-				CheckedPresentFailureCounter, MaxPresentFailures,
-				Result, *GetStateString()
-			);
-		}
-
-		if (IsValidRef(CustomPresent))
-		{
-			CustomPresent->PostPresent();
-		}
-
-#if LOG_PRESENT
-		const FString& ThreadName = FThreadManager::GetThreadName(FPlatformTLS::GetCurrentThreadId());
-		UE_LOG(LogD3D12RHI, Log, TEXT("*** PRESENT: Thread %s: Viewport %#016llx: BackBuffer %#016llx (SyncInterval %u) ***"), ThreadName.GetCharArray().GetData(), this, GetBackBuffer_RHIThread(), SyncInterval);
-#endif
+		PresentChecked(SyncInterval);
 	}
 
-	return bNeedNativePresent;
-}
-
-bool FD3D12Viewport::Present(FD3D12CommandContextBase& ContextBase, bool bLockToVsync)
-{
-	if (!IsPresentAllowed())
+	if (IsValidRef(CustomPresent))
 	{
-		return false;
+		CustomPresent->PostPresent();
 	}
 
-	check(ContextBase.GetParentAdapter() == GetParentAdapter());
-	
-	for (uint32 GPUIndex : FRHIGPUMask::All())
+	if (bNeedNativePresent || (IsValidRef(CustomPresent) && CustomPresent->NeedsAdvanceBackbuffer()))
 	{
-		FD3D12CommandContext& Context = *ContextBase.GetSingleDeviceContext(GPUIndex);
-
-		// Those are not necessarily the swap chain back buffer in case of multi-gpu
-		FD3D12Texture* DeviceBackBuffer = Context.RetrieveObject<FD3D12Texture, FRHITexture>(GetBackBuffer_RHIThread());
-
-		Context.FlushResourceBarriers();
-
-		// Currently, the swap chain Present() is called directly by the RHI thread.
-		// We need to submit the above commands and wait for the submission thread to process everything before we can continue.
-		Context.FlushCommands(ED3D12FlushFlags::WaitForSubmission);
-	}
-
-	const int32 SyncInterval = bLockToVsync ? RHIGetSyncInterval() : 0;
-	const bool bNativelyPresented = PresentChecked(ContextBase, SyncInterval);
-
-	if (bNativelyPresented || (CustomPresent && CustomPresent->NeedsAdvanceBackbuffer()))
-	{
 		// Increment back buffer
 #if DXGI_MAX_SWAPCHAIN_INTERFACE >= 3
-		if (bNativelyPresented && SwapChain3)
+		if (SwapChain3)
 		{
 			SetBackBufferIndex_RHIThread(SwapChain3->GetCurrentBackBufferIndex());
 		}
@@ -643,6 +692,105 @@
 		}
 #endif
 	}
+	
+	PresentArgs.PresentEvent->DispatchSubsequents();
+}
 
+void FD3D12Viewport::WaitForLastPresent() const
+{
+	FRenderThreadIdleScope Scope(ERenderThreadIdleTypes::WaitingForGPUPresent);
+
+	if ( (!PresentEvent.IsValid()) || (!PresentEvent->IsAwaitable()) || (PresentEvent->IsComplete()) ) 
+	{
+		return;
+	}
+		
+	{
+		TRACE_CPUPROFILER_EVENT_SCOPE_TEXT(TEXT("PresentEvent_Wait Frame"));
+		PresentEvent->Wait();
+	}
+	
+}
+
+bool FD3D12Viewport::EndRendering(FD3D12CommandContextBase& ContextBase, const FRHIPresentArgs& InPresentArgs)
+{
+#if CPUPROFILERTRACE_ENABLED
+	UE_TRACE_LOG_SCOPED_T(Cpu, End_Rendering, CpuChannel)
+		<< End_Rendering.BackBuffer(CurrentBackBufferIndex_RHIThread)
+		<< End_Rendering.PresentCounter(GRHIPresentCounter+1)
+		<< End_Rendering.FrameCounter(InPresentArgs.FrameCounter);
+#endif // CPUPROFILERTRACE_ENABLED
+	
+	if (!IsPresentAllowed() || !InPresentArgs.bPresent )
+	{
+		return false;
+	}
+
+	check(ContextBase.GetParentAdapter() == GetParentAdapter());
+
+	const bool bCustomPresent = IsValidRef(CustomPresent);
+	bool bNeedNativePresent = true;
+	bool bNativelyPresented = false;
+	FRHIPresentArgs LocalArgs = InPresentArgs;
+
+	if (CheckSwapChain())
+	{
+		if (GUseSubmissionThreadPresent)
+		{
+			SchedulePresent(&ContextBase, InPresentArgs);
+
+			// The native Present call has only been scheduled
+			bNativelyPresented = false;
+		}
+		else
+		{
+			SyncToSubmissionThread(ContextBase);
+
+			int32 SyncInterval = InPresentArgs.bLockToVsync ? RHIGetSyncInterval() : 0;
+
+			if (bCustomPresent)
+			{
+				SCOPE_CYCLE_COUNTER(STAT_D3D12CustomPresentTime);
+				// Give to the CustomPresent::Present the opportunity to use the RHIContext from the RHIThread
+				// CustomPresent::Present can decide to change the SyncInterval to toggle the Vsync
+				bNeedNativePresent = CustomPresent->Present(ContextBase, SyncInterval);
+			}
+
+			if (bNeedNativePresent)
+			{
+				PresentChecked(SyncInterval);
+				bNativelyPresented = true;
+			}
+
+			if (bCustomPresent)
+			{
+				CustomPresent->PostPresent();
+			}
+
+			if (bNativelyPresented || (bCustomPresent && CustomPresent->NeedsAdvanceBackbuffer()))
+			{
+				// Increment back buffer
+#if DXGI_MAX_SWAPCHAIN_INTERFACE >= 3
+				if (bNativelyPresented && SwapChain3)
+				{
+					SetBackBufferIndex_RHIThread(SwapChain3->GetCurrentBackBufferIndex());
+				}
+				else
+#endif
+				{
+					SetBackBufferIndex_RHIThread(CurrentBackBufferIndex_RHIThread + 1);
+				}
+
+#if !UE_BUILD_SHIPPING
+				if (RHIConsoleVariables::LogViewportEvents)
+				{
+					const FString& ThreadName = FThreadManager::GetThreadName(FPlatformTLS::GetCurrentThreadId());
+					UE_LOG(LogD3D12RHI, Log, TEXT("Thread %s: Incrementing RHIThread back buffer index of viewport: %#016llx to value: %u BackBuffer %#016llx"), ThreadName.GetCharArray().GetData(), this, CurrentBackBufferIndex_RHIThread, CurrentBackBuffer_RHIThread->Texture.GetReference());
+				}
+#endif
+			}
+		}
+	} 
+
 	return bNativelyPresented;
 }
@@ -750,9 +898,10 @@
  *	Viewport functions.
  *=============================================================================*/
 
-void FD3D12CommandContextBase::RHIEndDrawingViewport(FRHIViewport* ViewportRHI, bool bPresent, bool bLockToVsync)
+void FD3D12CommandContextBase::RHIEndDrawingViewport(FRHIViewport* ViewportRHI, const FRHIPresentArgs& InPresentArgs)
 {
 	FD3D12Viewport* Viewport = FD3D12DynamicRHI::ResourceCast(ViewportRHI);
+	SCOPE_CYCLE_COUNTER(STAT_D3D12PresentTime);
 
 #if !UE_BUILD_SHIPPING
 	if (RHIConsoleVariables::LogViewportEvents)
@@ -767,30 +916,32 @@
 	}
 #endif
 
-	SCOPE_CYCLE_COUNTER(STAT_D3D12PresentTime);
+	const bool bNativelyPresented = Viewport->EndRendering(*this, InPresentArgs);
 
-	bool bNativelyPresented = true;
-	if (bPresent)
+	if ((bNativelyPresented) || (GUseSubmissionThreadPresent))
 	{
-		bNativelyPresented = Viewport->Present(*this, bLockToVsync);
-	}
-	
-	// Multi-GPU support : here each GPU wait's for it's own frame completion.
-	if (bNativelyPresented)
-	{
+		FRenderThreadIdleScope Scope(ERenderThreadIdleTypes::WaitingForAllOtherSleep);
 		static const auto CFinishFrameVar = IConsoleManager::Get().FindTConsoleVariableDataInt(TEXT("r.FinishCurrentFrame"));
-		if (!CFinishFrameVar->GetValueOnRenderThread())
+
+		int32 FinishCurrentFrame = CFinishFrameVar->GetValueOnRenderThread();
+
+		if (FinishCurrentFrame == 0)
 		{
 			// Wait for the GPU to finish rendering the previous frame before finishing this frame.
 			Viewport->WaitForFrameEventCompletion();
 			Viewport->IssueFrameEvent();
 		}
-		else
+		else if (FinishCurrentFrame > 0)
 		{
 			// Finish current frame immediately to reduce latency
 			Viewport->IssueFrameEvent();
 			Viewport->WaitForFrameEventCompletion();
 		}
+		else if (FinishCurrentFrame < 0)
+		{
+			//No frame wait
+		}
+
 	}
 
 	// If the input latency timer has been triggered, block until the GPU is completely
@@ -858,3 +1009,51 @@
 
 	return SelectedBackBuffer;
 }
+
+void FD3D12Viewport::SyncToSubmissionThread(FD3D12CommandContextBase& ContextBase)
+{
+	// Sync with the submission thread 
+	for (uint32 GPUIndex : FRHIGPUMask::All())
+	{
+		FD3D12CommandContext& Context = *ContextBase.GetSingleDeviceContext(GPUIndex);
+
+		Context.FlushResourceBarriers();
+
+		// Currently, the swap chain Present() is called directly by the RHI thread.
+		// We need to submit the above commands and wait for the submission thread to process everything before we can continue.
+		Context.FlushCommands(ED3D12FlushFlags::WaitForSubmission);
+	}
+}
+
+bool FD3D12Viewport::CheckSwapChain()
+{
+#if PLATFORM_WINDOWS
+	// We can't call Present if !bIsValid, as it waits a window message to be processed, but the main thread may not be pumping the message handler.
+	if (bIsValid && SwapChain1.IsValid())
+	{
+		// Check if the viewport's swap chain has been invalidated by DXGI.
+		BOOL bSwapChainFullscreenState;
+		TRefCountPtr<IDXGIOutput> SwapChainOutput;
+		SwapChain1->GetFullscreenState(&bSwapChainFullscreenState, SwapChainOutput.GetInitReference());
+		// Can't compare BOOL with bool...
+		if ((!!bSwapChainFullscreenState) != bIsFullscreen)
+		{
+			bFullscreenLost = true;
+			bIsValid = false;
+		}
+	}
+
+	if (!bIsValid)
+	{
+#if WITH_MGPU
+		// Present failed so current expected GPU index will not match anymore, so patch up expected back buffer index
+		// Warning: Present is skipped for this frame but could cause a black screen for the next frame as well
+		FScopeLock Lock(&ExpectedBackBufferIndexLock);
+		SetBackBufferIndex_RenderThread((ExpectedBackBufferIndex_RenderThread == 0) ? NumBackBuffers - 1 : ExpectedBackBufferIndex_RenderThread - 1);
+#endif // WITH_MGPU
+		return false;
+	}
+#endif
+
+	return true;
+}
diff --git a/Engine/Source/Runtime/D3D12RHI/Private/D3D12Viewport.h b/Engine/Source/Runtime/D3D12RHI/Private/D3D12Viewport.h
--- a/Engine/Source/Runtime/D3D12RHI/Private/D3D12Viewport.h
+++ b/Engine/Source/Runtime/D3D12RHI/Private/D3D12Viewport.h
@@ -8,6 +8,7 @@
 
 #include "D3D12RHICommon.h"
 #include "D3D12Texture.h"
+#include "D3D12Submission.h"
 #include "HAL/Runnable.h"
 #include "MultiGPU.h"
 #include "RHIResources.h"
@@ -44,10 +45,27 @@ public:
 	 */
 	void ConditionalResetSwapChain(bool bIgnoreFocus);
 
+	/** Schedule Present from the RHI Thread to the Submission Thread
+	* Returns true if Present was correctly setup.
+	*/
+	bool SchedulePresent(class FD3D12CommandContextBase* BaseContext, const FRHIPresentArgs& InPresentArgs);
+	
+	/** Presents the swap chain on the Submission Thread
+	*/
+	void PresentOnSubmissionThread(FD3D12PresentArgs& PresentArgs);
+
+	/** Wait for the last Present to be called from the Submission Thread
+	*/
+	void WaitForLastPresent() const;
+
+	/** Signal that the Presents on the device has been completed and the backbuffer index has been update
+	*/
+	FGraphEventRef			PresentEvent;
+
 	/** Presents the swap chain.
 	 * Returns true if Present was done by Engine.
 	 */
-	bool Present(class FD3D12CommandContextBase& Context, bool bLockToVsync);
+	bool EndRendering(class FD3D12CommandContextBase& Context, const FRHIPresentArgs& InPresentArgs);
 
 	// Accessors.
 	FIntPoint GetSizeXY() const
@@ -74,16 +92,18 @@ public:
 	}
 #endif
 
-	FD3D12Texture* GetBackBuffer_RHIThread() const
-	{
-		checkSlow(CurrentBackBuffer_RHIThread);
-		return CurrentBackBuffer_RHIThread->Texture;
-	}
+	FD3D12Texture* GetBackBuffer_RHIThread(bool bWaitForLastPresent = true) const;
+
 
-	FD3D12Texture* GetSDRBackBuffer_RHIThread() const
+	FD3D12Texture* GetSDRBackBuffer_RHIThread(bool bWaitForLastPresent = true) const
 	{
 		checkSlow(CurrentBackBuffer_RHIThread);
 
+		if (bWaitForLastPresent)
+		{
+			WaitForLastPresent();
+		}
+
 #if D3D12RHI_USE_SDR_BACKBUFFER
 		if (PixelFormat != SDRPixelFormat)
 		{
@@ -148,6 +168,8 @@ public:
 
 private:
 	bool IsPresentAllowed();
+	bool CheckSwapChain();
+	void SyncToSubmissionThread(FD3D12CommandContextBase& ContextBase);
 
 	FString GetStateString();
 
@@ -157,9 +179,8 @@ private:
 
 	/**
 	 * Presents the swap chain checking the return result.
-	 * Returns true if Present was done by Engine.
 	 */
-	bool PresentChecked(IRHICommandContext& RHICmdContext, int32 SyncInterval);
+	void PresentChecked(int32 SyncInterval);
 
 	/**
 	 * Presents the backbuffer to the viewport window.
diff --git a/Engine/Source/Runtime/Engine/Private/UnrealClient.cpp b/Engine/Source/Runtime/Engine/Private/UnrealClient.cpp
--- a/Engine/Source/Runtime/Engine/Private/UnrealClient.cpp
+++ b/Engine/Source/Runtime/Engine/Private/UnrealClient.cpp
@@ -1667,7 +1667,7 @@ void FViewport::EndRenderFrame(FRHICommandListImmediate& RHICmdList, bool bPrese
 		UEngine::SetPresentLatencyMarkerStart(CurrentFrameCounter);
 	});
 
-	RHICmdList.EndDrawingViewport(GetViewportRHI(), bPresent, bLockToVsync);
+	RHICmdList.EndDrawingViewport(GetViewportRHI(), FRHIPresentArgs(/*.FrameCounter*/ GFrameCounterRenderThread, /*.bPresent*/ bPresent, /*.bLockToVsync*/ bLockToVsync));
 
 	RHICmdList.EnqueueLambda([CurrentFrameCounter = GFrameCounterRenderThread](FRHICommandListImmediate& InRHICmdList)
 	{
diff --git a/Engine/Source/Runtime/NullDrv/Public/NullRHI.h b/Engine/Source/Runtime/NullDrv/Public/NullRHI.h
--- a/Engine/Source/Runtime/NullDrv/Public/NullRHI.h
+++ b/Engine/Source/Runtime/NullDrv/Public/NullRHI.h
@@ -420,7 +420,7 @@ public:
 		return true; 
 	}
 
-	virtual void RHIEndDrawingViewport(FRHIViewport* Viewport, bool bPresent, bool bLockToVsync) final override
+	virtual void RHIEndDrawingViewport(FRHIViewport* Viewport, const FRHIPresentArgs& InPresentArgs) final override
 	{
 	}
 
diff --git a/Engine/Source/Runtime/OpenGLDrv/Private/OpenGLDrv.h b/Engine/Source/Runtime/OpenGLDrv/Private/OpenGLDrv.h
--- a/Engine/Source/Runtime/OpenGLDrv/Private/OpenGLDrv.h
+++ b/Engine/Source/Runtime/OpenGLDrv/Private/OpenGLDrv.h
@@ -410,7 +410,7 @@ public:
 	UE_API virtual void RHIBeginRenderQuery(FRHIRenderQuery* RenderQuery) final override;
 	UE_API virtual void RHIEndRenderQuery(FRHIRenderQuery* RenderQuery) final override;
 
-	UE_API virtual void RHIEndDrawingViewport(FRHIViewport* Viewport, bool bPresent, bool bLockToVsync) final override;
+	UE_API virtual void RHIEndDrawingViewport(FRHIViewport* Viewport, const FRHIPresentArgs& InPresentArgs) final override;
 	UE_API virtual void RHIEndFrame(const FRHIEndFrameArgs& Args) final override;
 	UE_API virtual void RHISetStreamSource(uint32 StreamIndex, FRHIBuffer* VertexBuffer, uint32 Offset) final override;
 	UE_API virtual void RHISetRasterizerState(FRHIRasterizerState* NewState) final override;
diff --git a/Engine/Source/Runtime/OpenGLDrv/Private/OpenGLViewport.cpp b/Engine/Source/Runtime/OpenGLDrv/Private/OpenGLViewport.cpp
--- a/Engine/Source/Runtime/OpenGLDrv/Private/OpenGLViewport.cpp
+++ b/Engine/Source/Runtime/OpenGLDrv/Private/OpenGLViewport.cpp
@@ -109,7 +109,7 @@ void FOpenGLDynamicRHI::RHITick( float DeltaTime )
  *	Viewport functions.
  *=============================================================================*/
 
-void FOpenGLDynamicRHI::RHIEndDrawingViewport(FRHIViewport* ViewportRHI,bool bPresent,bool bLockToVsync)
+void FOpenGLDynamicRHI::RHIEndDrawingViewport(FRHIViewport* ViewportRHI, const FRHIPresentArgs& InPresentArgs)
 {
 	VERIFY_GL_SCOPE();
 
@@ -131,8 +131,8 @@ void FOpenGLDynamicRHI::RHIEndDrawingViewport(FRHIViewport* ViewportRHI,bool bPr
 			*Viewport, 
 			BackBuffer->GetSizeX(),
 			BackBuffer->GetSizeY(),
-			bPresent,
-			bLockToVsync
+			InPresentArgs.bPresent,
+			InPresentArgs.bLockToVsync
 		);
 
 		// Always consider the Framebuffer in the rendering context dirty after the blit
@@ -141,18 +141,24 @@ void FOpenGLDynamicRHI::RHIEndDrawingViewport(FRHIViewport* ViewportRHI,bool bPr
 		if (bNeedFinishFrame)
 		{
 			static const auto CFinishFrameVar = IConsoleManager::Get().FindTConsoleVariableDataInt(TEXT("r.FinishCurrentFrame"));
-			if (!CFinishFrameVar->GetValueOnRenderThread())
+			int32 FinishCurrentFrame = CFinishFrameVar->GetValueOnRenderThread();
+
+			if (FinishCurrentFrame == 0)
 			{
 				// Wait for the GPU to finish rendering the previous frame before finishing this frame.
 				Viewport->WaitForFrameEventCompletion();
 				Viewport->IssueFrameEvent();
 			}
-			else
+			else if (FinishCurrentFrame > 0)
 			{
 				// Finish current frame immediately to reduce latency
 				Viewport->IssueFrameEvent();
 				Viewport->WaitForFrameEventCompletion();
 			}
+			else if (FinishCurrentFrame < 0)
+			{
+				//No frame wait
+			}
 		}
 		
 		// If the input latency timer has been triggered, block until the GPU is completely
diff --git a/Engine/Source/Runtime/RHI/Private/RHICommandList.cpp b/Engine/Source/Runtime/RHI/Private/RHICommandList.cpp
--- a/Engine/Source/Runtime/RHI/Private/RHICommandList.cpp
+++ b/Engine/Source/Runtime/RHI/Private/RHICommandList.cpp
@@ -1900,7 +1900,7 @@ void FRHICommandListImmediate::UnStallRHIThread()
 	}
 }
 
-void FRHICommandListImmediate::EndDrawingViewport(FRHIViewport* Viewport, bool bPresent, bool bLockToVsync)
+void FRHICommandListImmediate::EndDrawingViewport(FRHIViewport* Viewport, const FRHIPresentArgs& InPresentArgs)
 {
 	// Make sure all prior graphics and async compute work has been submitted.
 	// This is necessary because platform RHIs often submit additional work on the graphics queue during present, and we need to ensure we won't deadlock on async work that wasn't yet submitted by the renderer.
@@ -1910,11 +1910,11 @@ void FRHICommandListImmediate::EndDrawingViewport(FRHIViewport* Viewport, bool b
 	check(IsImmediate() && IsInRenderingThread());
 	if (Bypass())
 	{
-		GetContext().RHIEndDrawingViewport(Viewport, bPresent, bLockToVsync);
+		GetContext().RHIEndDrawingViewport(Viewport, InPresentArgs);
 	}
 	else
 	{
-		ALLOC_COMMAND(FRHICommandEndDrawingViewport)(Viewport, bPresent, bLockToVsync);
+		ALLOC_COMMAND(FRHICommandEndDrawingViewport)(Viewport, InPresentArgs);
 
 		// if we aren't running an RHIThread, there is no good reason to buffer this frame advance stuff and that complicates state management, so flush everything out now
 		{
@@ -1923,7 +1923,7 @@ void FRHICommandListImmediate::EndDrawingViewport(FRHIViewport* Viewport, bool b
 		}
 	}
 
-	GDynamicRHI->RHIAdvanceFrameForGetViewportBackBuffer(Viewport, bPresent);
+	GDynamicRHI->RHIAdvanceFrameForGetViewportBackBuffer(Viewport, InPresentArgs.bPresent);
 }
 
 void FRHICommandListImmediate::EndFrame()
diff --git a/Engine/Source/Runtime/RHI/Public/DynamicRHI.h b/Engine/Source/Runtime/RHI/Public/DynamicRHI.h
--- a/Engine/Source/Runtime/RHI/Public/DynamicRHI.h
+++ b/Engine/Source/Runtime/RHI/Public/DynamicRHI.h
@@ -1263,6 +1263,7 @@ inline uint32 RHIComputeMemorySize(FRHITexture* TextureRHI)
 
 inline bool RHIGetRenderQueryResult(FRHIRenderQuery* RenderQuery, uint64& OutResult, bool bWait, uint32 GPUIndex = INDEX_NONE)
 {
+	TRACE_CPUPROFILER_EVENT_SCOPE_CONDITIONAL(RHIGetRenderQueryResult_GPU_Wait, bWait);
 	return GDynamicRHI->RHIGetRenderQueryResult(RenderQuery, OutResult, bWait, GPUIndex);
 }
 
diff --git a/Engine/Source/Runtime/RHI/Public/RHICommandList.h b/Engine/Source/Runtime/RHI/Public/RHICommandList.h
--- a/Engine/Source/Runtime/RHI/Public/RHICommandList.h
+++ b/Engine/Source/Runtime/RHI/Public/RHICommandList.h
@@ -2514,13 +2514,11 @@ FRHICOMMAND_MACRO(FRHICommandPostExternalCommandsReset)
 FRHICOMMAND_MACRO(FRHICommandEndDrawingViewport)
 {
 	FRHIViewport* Viewport;
-	bool bPresent;
-	bool bLockToVsync;
+	FRHIPresentArgs PresentArgs;
 
-	inline FRHICommandEndDrawingViewport(FRHIViewport* InViewport, bool InbPresent, bool InbLockToVsync)
+	inline FRHICommandEndDrawingViewport(FRHIViewport* InViewport, const FRHIPresentArgs& InPresentArgs)
 		: Viewport(InViewport)
-		, bPresent(InbPresent)
-		, bLockToVsync(InbLockToVsync)
+		, PresentArgs(InPresentArgs)
 	{
 	}
 	RHI_EXECUTE_API void Execute(FRHICommandListBase& CmdList);
@@ -4657,7 +4655,13 @@ public:
 
 	UE_DEPRECATED(5.7, "FRHICommandListImmediate::BeginDrawingViewport is deprecated. Calls to this function are unnecessary and can be removed. There is no replacement.")
 	inline void BeginDrawingViewport(FRHIViewport* Viewport, FRHITexture* RenderTargetRHI) {}
-	RHI_API void EndDrawingViewport(FRHIViewport* Viewport, bool bPresent, bool bLockToVsync);
+
+	UE_DEPRECATED(5.7, "FRHICommandListImmediate::EndDrawingViewport(FRHIViewport* Viewport, bool bPresent, bool bLockToVsync) is deprecated. Use FRHIPresentArgs to pass the parameters needed.")
+	inline void EndDrawingViewport(FRHIViewport* Viewport, bool bPresent, bool bLockToVsync)
+	{
+		EndDrawingViewport(Viewport, FRHIPresentArgs(0, bPresent, bLockToVsync));
+	}
+	RHI_API void EndDrawingViewport(FRHIViewport* Viewport, const FRHIPresentArgs& InPresentArgs);
 
 	RHI_API void EndFrame();
 
diff --git a/Engine/Source/Runtime/RHI/Public/RHICommandListCommandExecutes.inl b/Engine/Source/Runtime/RHI/Public/RHICommandListCommandExecutes.inl
--- a/Engine/Source/Runtime/RHI/Public/RHICommandListCommandExecutes.inl
+++ b/Engine/Source/Runtime/RHI/Public/RHICommandListCommandExecutes.inl
@@ -562,5 +562,5 @@ void FRHICommandSetBindingsOnShaderBindingTable::Execute(FRHICommandListBase& Cm
 void FRHICommandEndDrawingViewport::Execute(FRHICommandListBase& CmdList)
 {
 	RHISTAT(EndDrawingViewport);
-	INTERNAL_DECORATOR(RHIEndDrawingViewport)(Viewport, bPresent, bLockToVsync);
+	INTERNAL_DECORATOR(RHIEndDrawingViewport)(Viewport, PresentArgs);
 }
diff --git a/Engine/Source/Runtime/RHI/Public/RHIContext.h b/Engine/Source/Runtime/RHI/Public/RHIContext.h
--- a/Engine/Source/Runtime/RHI/Public/RHIContext.h
+++ b/Engine/Source/Runtime/RHI/Public/RHIContext.h
@@ -731,7 +731,7 @@ public:
 #endif
 
 	// This method is queued with an RHIThread, otherwise it will flush after it is queued; without an RHI thread there is no benefit to queuing this frame advance commands
-	virtual void RHIEndDrawingViewport(FRHIViewport* Viewport, bool bPresent, bool bLockToVsync) = 0;
+	virtual void RHIEndDrawingViewport(FRHIViewport* Viewport, const FRHIPresentArgs& InPresentArgs) = 0;
 
 	virtual void RHISetStreamSource(uint32 StreamIndex, FRHIBuffer* VertexBuffer, uint32 Offset) = 0;
 
diff --git a/Engine/Source/Runtime/RHI/Public/RHIResources.h b/Engine/Source/Runtime/RHI/Public/RHIResources.h
--- a/Engine/Source/Runtime/RHI/Public/RHIResources.h
+++ b/Engine/Source/Runtime/RHI/Public/RHIResources.h
@@ -2511,6 +2511,25 @@ inline FRHIPooledRenderQuery::~FRHIPooledRenderQuery()
 	ReleaseQuery();
 }
 
+struct FRHIPresentArgs
+{
+	FRHIPresentArgs(uint64 InFrameCounter, bool bInPresent, bool bInLockToVsync):
+		FrameCounter(InFrameCounter),
+		bPresent((uint8)bInPresent),
+		bLockToVsync((uint8)bInLockToVsync)
+	{ }
+
+
+	/** Frame index that is going to Presented. */
+	uint64 FrameCounter;
+
+	/** true if the Engine should call Present the frame to the display */
+	uint8 bPresent:1;
+
+	/** true if the Engine should sync the frame presentation to the display vsync*/
+	uint8 bLockToVsync:1;
+};
+
 class FRHIViewport : public FRHIResource 
 {
 public:
@@ -4321,7 +4340,14 @@ public:
 	// Present.  Must match value previously returned by NeedsNativePresent for this frame.
 	virtual bool Present(IRHICommandContext& RHICmdContext, int32& InOutSyncInterval) { check(false); return true; };
 
-	// Called from RHI thread after native Present has been called
+	// Called from Submission thread to perform custom present.
+	// @param InOutSyncInterval - in out param, indicates if vsync is on (>0) or off (==0).
+	// @return	true if native Present should be also be performed; false otherwise. If it returns
+	// true, then InOutSyncInterval could be modified to switch between VSync/NoVSync for the normal 
+	// Present.  Must match value previously returned by NeedsNativePresent for this frame.
+	virtual bool PresentOnSubmissionThread(int32& InOutSyncInterval) { check(false); return true; };
+
+	// Called after native Present has been called
 	virtual void PostPresent() {};
 
 	// Called when rendering thread is acquired
diff --git a/Engine/Source/Runtime/RHI/Public/RHIValidationContext.h b/Engine/Source/Runtime/RHI/Public/RHIValidationContext.h
--- a/Engine/Source/Runtime/RHI/Public/RHIValidationContext.h
+++ b/Engine/Source/Runtime/RHI/Public/RHIValidationContext.h
@@ -758,9 +758,9 @@ public:
 #endif
 
 	// This method is queued with an RHIThread, otherwise it will flush after it is queued; without an RHI thread there is no benefit to queuing this frame advance commands
-	virtual void RHIEndDrawingViewport(FRHIViewport* Viewport, bool bPresent, bool bLockToVsync) override final
+	virtual void RHIEndDrawingViewport(FRHIViewport* Viewport, const FRHIPresentArgs& InPresentArgs) override final
 	{
-		RHIContext->RHIEndDrawingViewport(Viewport, bPresent, bLockToVsync);
+		RHIContext->RHIEndDrawingViewport(Viewport, InPresentArgs);
 	}
 
 	virtual void RHISetStreamSource(uint32 StreamIndex, FRHIBuffer* VertexBuffer, uint32 Offset) override final
diff --git a/Engine/Source/Runtime/RenderCore/Private/RenderingThread.cpp b/Engine/Source/Runtime/RenderCore/Private/RenderingThread.cpp
--- a/Engine/Source/Runtime/RenderCore/Private/RenderingThread.cpp
+++ b/Engine/Source/Runtime/RenderCore/Private/RenderingThread.cpp
@@ -2477,6 +2477,7 @@ namespace FFrameEndSync
 
 	void Sync(EFlushMode FlushMode)
 	{
+		TRACE_CPUPROFILER_EVENT_SCOPE(Sync_RenderingThread);
 		static bool bRecursive = false;
 		TGuardValue<bool> RecursionGuard(bRecursive, true);
 
diff --git a/Engine/Source/Runtime/SlateRHIRenderer/Private/SlateRHIRenderer.cpp b/Engine/Source/Runtime/SlateRHIRenderer/Private/SlateRHIRenderer.cpp
--- a/Engine/Source/Runtime/SlateRHIRenderer/Private/SlateRHIRenderer.cpp
+++ b/Engine/Source/Runtime/SlateRHIRenderer/Private/SlateRHIRenderer.cpp
@@ -941,7 +941,7 @@ void FSlateRHIRenderer::PresentWindow_RenderThread(FRHICommandListImmediate& RHI
 		}
 	}
 
-	RHICmdList.EndDrawingViewport(DrawPassOutputs.ViewportRHI, true, DrawPassInputs.bLockToVsync);
+	RHICmdList.EndDrawingViewport(DrawPassOutputs.ViewportRHI, FRHIPresentArgs(/*.FrameCounter*/ GFrameCounterRenderThread, /*.bPresent*/ true, /*.bLockToVsync*/ DrawPassInputs.bLockToVsync));
 
 	RHICmdList.EnqueueLambda([CurrentFrameCounter = GFrameCounterRenderThread](FRHICommandListImmediate& InRHICmdList)
 	{
diff --git a/Engine/Source/Runtime/VulkanRHI/Private/VulkanContext.h b/Engine/Source/Runtime/VulkanRHI/Private/VulkanContext.h
--- a/Engine/Source/Runtime/VulkanRHI/Private/VulkanContext.h
+++ b/Engine/Source/Runtime/VulkanRHI/Private/VulkanContext.h
@@ -338,7 +338,7 @@ public:
 
 	virtual void Finalize(TArray<FVulkanPayload*>& OutPayloads) override;
 
-	virtual void RHIEndDrawingViewport(FRHIViewport* Viewport, bool bPresent, bool bLockToVsync) final override;
+	virtual void RHIEndDrawingViewport(FRHIViewport* Viewport, const FRHIPresentArgs& InPresentArgs) final override;
 
 	virtual void RHIBeginRenderPass(const FRHIRenderPassInfo& InInfo, const TCHAR* InName) final override;
 	virtual void RHIEndRenderPass() final override;
diff --git a/Engine/Source/Runtime/VulkanRHI/Private/VulkanRHI.cpp b/Engine/Source/Runtime/VulkanRHI/Private/VulkanRHI.cpp
--- a/Engine/Source/Runtime/VulkanRHI/Private/VulkanRHI.cpp
+++ b/Engine/Source/Runtime/VulkanRHI/Private/VulkanRHI.cpp
@@ -1281,7 +1281,7 @@ FVulkanTiming::FVulkanTiming(FVulkanQueue& InQueue)
 }
 #endif
 
-void FVulkanCommandListContext::RHIEndDrawingViewport(FRHIViewport* ViewportRHI, bool bPresent, bool bLockToVsync)
+void FVulkanCommandListContext::RHIEndDrawingViewport(FRHIViewport* ViewportRHI, const FRHIPresentArgs& InPresentArgs)
 {
 	LLM_SCOPE_VULKAN(ELLMTagVulkan::VulkanMisc);
 	//FRCLog::Printf(FString::Printf(TEXT("FVulkanCommandListContext::RHIEndDrawingViewport()")));
@@ -1291,7 +1291,7 @@ void FVulkanCommandListContext::RHIEndDrawingViewport(FRHIViewport* ViewportRHI,
 	FVulkanCommandBuffer& CommandBuffer = GetCommandBuffer();
 	check(!CommandBuffer.HasEnded() && !CommandBuffer.IsInsideRenderPass());
 
-	const bool bNativePresent = Viewport->Present(*this, bLockToVsync);
+	const bool bNativePresent = Viewport->Present(*this, InPresentArgs.bLockToVsync);
 	if (bNativePresent)
 	{
 		//#todo-rco: Check for r.FinishCurrentFrame
diff --git a/Engine/Source/Runtime/Windows/D3D11RHI/Private/D3D11RHIPrivate.h b/Engine/Source/Runtime/Windows/D3D11RHI/Private/D3D11RHIPrivate.h
--- a/Engine/Source/Runtime/Windows/D3D11RHI/Private/D3D11RHIPrivate.h
+++ b/Engine/Source/Runtime/Windows/D3D11RHI/Private/D3D11RHIPrivate.h
@@ -537,7 +537,7 @@ public:
 	D3D11RHI_API virtual void RHIBeginRenderQuery(FRHIRenderQuery* RenderQuery) final override;
 	D3D11RHI_API virtual void RHIEndRenderQuery  (FRHIRenderQuery* RenderQuery) final override;
 
-	D3D11RHI_API virtual void RHIEndDrawingViewport(FRHIViewport* Viewport, bool bPresent, bool bLockToVsync) final override;
+	D3D11RHI_API virtual void RHIEndDrawingViewport(FRHIViewport* Viewport, const FRHIPresentArgs& InPresentArgs) final override;
 	D3D11RHI_API virtual void RHIEndFrame(const FRHIEndFrameArgs& Args) final override;
 	D3D11RHI_API virtual void RHISetStreamSource(uint32 StreamIndex, FRHIBuffer* VertexBuffer, uint32 Offset) final override;
 	D3D11RHI_API virtual void RHISetRasterizerState(FRHIRasterizerState* NewState) final override;
diff --git a/Engine/Source/Runtime/Windows/D3D11RHI/Private/D3D11Viewport.cpp b/Engine/Source/Runtime/Windows/D3D11RHI/Private/D3D11Viewport.cpp
--- a/Engine/Source/Runtime/Windows/D3D11RHI/Private/D3D11Viewport.cpp
+++ b/Engine/Source/Runtime/Windows/D3D11RHI/Private/D3D11Viewport.cpp
@@ -639,6 +639,8 @@ void FD3D11Viewport::PresentWithVsyncDWM(IRHICommandContext& RHICmdContext)
 
 bool FD3D11Viewport::Present(IRHICommandContext& RHICmdContext, bool bLockToVsync)
 {
+	TRACE_CPUPROFILER_EVENT_SCOPE(FD3D11Viewport::Present);
+
 	bool bNativelyPresented = true;
 #if	D3D11_WITH_DWMAPI
 	// We can't call Present if !bIsValid, as it waits a window message to be processed, but the main thread may not be pumping the message handler.
@@ -753,7 +755,7 @@ void FD3D11DynamicRHI::RHITick( float DeltaTime )
  *	Viewport functions.
  *=============================================================================*/
 
-void FD3D11DynamicRHI::RHIEndDrawingViewport(FRHIViewport* ViewportRHI,bool bPresent,bool bLockToVsync)
+void FD3D11DynamicRHI::RHIEndDrawingViewport(FRHIViewport* ViewportRHI, const FRHIPresentArgs& InPresentArgs)
 {
 	++PresentCounter;
 	FD3D11Viewport* Viewport = ResourceCast(ViewportRHI);
@@ -794,26 +796,32 @@ void FD3D11DynamicRHI::RHIEndDrawingViewport(FRHIViewport* ViewportRHI,bool bPre
 	// Compute Shader is set to NULL after each Dispatch call, so no need to clear it here
 
 	bool bNativelyPresented = true;
-	if (bPresent)
+	if (InPresentArgs.bPresent)
 	{
-		bNativelyPresented = Viewport->Present(*this, bLockToVsync);
+		bNativelyPresented = Viewport->Present(*this, InPresentArgs.bLockToVsync);
 	}
 
 	if (bNativelyPresented)
 	{ 
 		static const auto CFinishFrameVar = IConsoleManager::Get().FindTConsoleVariableDataInt(TEXT("r.FinishCurrentFrame"));
-		if (!CFinishFrameVar->GetValueOnRenderThread())
+		int32 FinishCurrentFrame = CFinishFrameVar->GetValueOnRenderThread();
+
+		if (FinishCurrentFrame == 0)
 		{
 			// Wait for the GPU to finish rendering the previous frame before finishing this frame.
 			Viewport->WaitForFrameEventCompletion();
 			Viewport->IssueFrameEvent();
 		}
-		else
+		else if (FinishCurrentFrame > 0)
 		{
 			// Finish current frame immediately to reduce latency
 			Viewport->IssueFrameEvent();
 			Viewport->WaitForFrameEventCompletion();
 		}
+		else if (FinishCurrentFrame < 0)
+		{
+			//No frame wait
+		}
 	}
 
 	// If the input latency timer has been triggered, block until the GPU is completely
