[D3D12RHI]
* strnlen_s -> FPlatformString::Strnlen

diff --git a/Engine/Source/Runtime/D3D12RHI/Private/Windows/WindowsD3D12PipelineState.cpp b/Engine/Source/Runtime/D3D12RHI/Private/Windows/WindowsD3D12PipelineState.cpp
--- a/Engine/Source/Runtime/D3D12RHI/Private/Windows/WindowsD3D12PipelineState.cpp
+++ b/Engine/Source/Runtime/D3D12RHI/Private/Windows/WindowsD3D12PipelineState.cpp
@@ -364,7 +364,7 @@ void FD3D12PipelineStateCache::AddToDiskCache(const FD3D12LowLevelGraphicsPipeli
 			for (uint32 i = 0; i < PsoDesc.InputLayout.NumElements; i++)
 			{
 				//Save the Sematic name string
-				uint32 stringLength = (uint32)strnlen_s(PsoDesc.InputLayout.pInputElementDescs[i].SemanticName, IL_MAX_SEMANTIC_NAME);
+				uint32 stringLength = (uint32)FPlatformString::Strnlen(PsoDesc.InputLayout.pInputElementDescs[i].SemanticName, IL_MAX_SEMANTIC_NAME);
 				stringLength++; // include the NULL char
 				DiskCache.AppendData((void*)&stringLength, sizeof(stringLength));
 				DiskCache.AppendData((void*)PsoDesc.InputLayout.pInputElementDescs[i].SemanticName, stringLength);
