Fix error compiling RHI on some platforms due to missing _RHI_RAYTRACING branches.

diff --git a/Engine/Source/Runtime/Apple/MetalRHI/Private/MetalRHI.cpp b/Engine/Source/Runtime/Apple/MetalRHI/Private/MetalRHI.cpp
--- a/Engine/Source/Runtime/Apple/MetalRHI/Private/MetalRHI.cpp
+++ b/Engine/Source/Runtime/Apple/MetalRHI/Private/MetalRHI.cpp
@@ -1907,11 +1907,13 @@ void FMetalDynamicRHI::RHIWriteResourceGPUAddresses(FRHICommandListBase& RHICmdL
 				{
 					unimplemented();
 				}
+#if METAL_RHI_RAYTRACING
 				else if (Resource->GetType() == RRT_RayTracingGeometry)
 				{
 					const FMetalRayTracingGeometry* RayTracingGeometry = ResourceCast(static_cast<const FRHIRayTracingGeometry*>(Resource));
 					AddressesPtr[ResourceIndex] = RayTracingGeometry->GetGPUVirtualAddress();
 				}
+#endif
 				else
 				{
 					checkNoEntry();
diff --git a/Engine/Source/Runtime/D3D12RHI/Private/D3D12Resources.cpp b/Engine/Source/Runtime/D3D12RHI/Private/D3D12Resources.cpp
--- a/Engine/Source/Runtime/D3D12RHI/Private/D3D12Resources.cpp
+++ b/Engine/Source/Runtime/D3D12RHI/Private/D3D12Resources.cpp
@@ -1697,11 +1697,13 @@ void FD3D12DynamicRHI::RHIWriteResourceGPUAddresses(FRHICommandListBase& RHICmdL
 				{
 					unimplemented();
 				}
+#if D3D12_RHI_RAYTRACING
 				else if (Resource->GetType() == RRT_RayTracingGeometry)
 				{
 					const FD3D12RayTracingGeometry* RayTracingGeometry = ResourceCast(static_cast<const FRHIRayTracingGeometry*>(Resource));
 					AddressesPtr[ResourceIndex] = RayTracingGeometry->GetGPUVirtualAddress(GPUIndex);
 				}
+#endif
 				else
 				{
 					checkNoEntry();
