Adding UE_D3D12_SDK_VERSION which mirrors the Agility SDK D3D12_SDK_VERSION for all platforms.

diff --git a/Engine/Source/Runtime/D3D12RHI/Private/D3D12EnhancedBarriers.cpp b/Engine/Source/Runtime/D3D12RHI/Private/D3D12EnhancedBarriers.cpp
--- a/Engine/Source/Runtime/D3D12RHI/Private/D3D12EnhancedBarriers.cpp
+++ b/Engine/Source/Runtime/D3D12RHI/Private/D3D12EnhancedBarriers.cpp
@@ -269,7 +269,12 @@ static FString ConvertToString(D3D12_BARRIER_LAYOUT InLayout)
 		ENUMVAL(D3D12_BARRIER_LAYOUT_COMPUTE_QUEUE_SHADER_RESOURCE),
 		ENUMVAL(D3D12_BARRIER_LAYOUT_COMPUTE_QUEUE_COPY_SOURCE),
 		ENUMVAL(D3D12_BARRIER_LAYOUT_COMPUTE_QUEUE_COPY_DEST),
+#if UE_D3D12_SDK_VERSION < 618
 		ENUMVAL(D3D12_BARRIER_LAYOUT_VIDEO_QUEUE_COMMON),
+#endif
+#if UE_D3D12_SDK_VERSION >= 618
+		ENUMVAL(D3D12_BARRIER_LAYOUT_DIRECT_QUEUE_GENERIC_READ_COMPUTE_QUEUE_ACCESSIBLE),
+#endif
 		ENUMVAL(D3D12_BARRIER_LAYOUT_PRESENT),
 		ENUMVAL(D3D12_BARRIER_LAYOUT_UNDEFINED));
 }
diff --git a/Engine/Source/Runtime/D3D12RHI/Private/D3D12WorkGraph.cpp b/Engine/Source/Runtime/D3D12RHI/Private/D3D12WorkGraph.cpp
--- a/Engine/Source/Runtime/D3D12RHI/Private/D3D12WorkGraph.cpp
+++ b/Engine/Source/Runtime/D3D12RHI/Private/D3D12WorkGraph.cpp
@@ -248,7 +248,7 @@ FD3D12WorkGraphPipelineState::FD3D12WorkGraphPipelineState(FD3D12Device* Device,
 	RootArgStrideInBytes = ((RootArgStrideInBytes + 15) & ~15);
 	MaxRootArgOffset = RootArgOffsets.Last();
 
-#if D3D12_SDK_VERSION < 616
+#if UE_D3D12_SDK_VERSION < 616
 	WorkGraphSubobject->Finalize();
 #endif
 
diff --git a/Engine/Source/Runtime/D3D12RHI/Private/Windows/WindowsD3D12RHIDefinitions.h b/Engine/Source/Runtime/D3D12RHI/Private/Windows/WindowsD3D12RHIDefinitions.h
--- a/Engine/Source/Runtime/D3D12RHI/Private/Windows/WindowsD3D12RHIDefinitions.h
+++ b/Engine/Source/Runtime/D3D12RHI/Private/Windows/WindowsD3D12RHIDefinitions.h
@@ -4,6 +4,8 @@
 
 #include "D3D12ThirdParty.h"
 
+#define UE_D3D12_SDK_VERSION (D3D12_SDK_VERSION)
+
 #define D3D12RHI_DEFAULT_NUM_BACKBUFFER 3
 
 // Windows uses the graphics command list interface for the copy queue.
