Fix Automated OOB Access Bug Report With New DX Agility SDK

The new DX Agility SDK moved around some of the enumerator values for D3D12_BARRIER_LAYOUT. Despite not yet using the enumerators that were rearranged, our automated bug reporting felt that a related lookup table should be updated. This does that as well as fixes a bounds check when indexing that table.

#rnx

diff --git a/Engine/Source/Runtime/D3D12RHI/Private/D3D12EnhancedBarriers.cpp b/Engine/Source/Runtime/D3D12RHI/Private/D3D12EnhancedBarriers.cpp
--- a/Engine/Source/Runtime/D3D12RHI/Private/D3D12EnhancedBarriers.cpp
+++ b/Engine/Source/Runtime/D3D12RHI/Private/D3D12EnhancedBarriers.cpp
@@ -269,13 +269,15 @@ static FString ConvertToString(D3D12_BARRIER_LAYOUT InLayout)
 		ENUMVAL(D3D12_BARRIER_LAYOUT_COMPUTE_QUEUE_SHADER_RESOURCE),
 		ENUMVAL(D3D12_BARRIER_LAYOUT_COMPUTE_QUEUE_COPY_SOURCE),
 		ENUMVAL(D3D12_BARRIER_LAYOUT_COMPUTE_QUEUE_COPY_DEST),
-#if UE_D3D12_SDK_VERSION < 618
+#if D3D12RHI_HAS_LAYOUT_VIDEO_QUEUE_COMMON
 		ENUMVAL(D3D12_BARRIER_LAYOUT_VIDEO_QUEUE_COMMON),
 #endif
-#if UE_D3D12_SDK_VERSION >= 618
+#if D3D12RHI_HAS_LAYOUT_DIRECT_QUEUE_GENERIC_READ_COMPUTE_QUEUE_ACCESSIBLE
 		ENUMVAL(D3D12_BARRIER_LAYOUT_DIRECT_QUEUE_GENERIC_READ_COMPUTE_QUEUE_ACCESSIBLE),
 #endif
+#if D3D12RHI_COMMON_AND_PRESENT_ARE_DISTINCT
 		ENUMVAL(D3D12_BARRIER_LAYOUT_PRESENT),
+#endif
 		ENUMVAL(D3D12_BARRIER_LAYOUT_UNDEFINED));
 }
 
@@ -550,7 +552,12 @@ TStaticArray<D3D12_BARRIER_SYNC, 23> AccessCompatibleSync = {
 };
 
 static constexpr
-TStaticArray<D3D12_BARRIER_ACCESS, 32> LayoutCompatibleAccess = {
+uint32 LayoutCompatibleAccessNum = 31
+	+ !!(D3D12RHI_COMMON_AND_PRESENT_ARE_DISTINCT)
+	+ !!(D3D12RHI_HAS_LAYOUT_DIRECT_QUEUE_GENERIC_READ_COMPUTE_QUEUE_ACCESSIBLE);
+
+static constexpr
+TStaticArray<D3D12_BARRIER_ACCESS, LayoutCompatibleAccessNum> LayoutCompatibleAccess = {
 	//D3D12_BARRIER_LAYOUT_COMMON
 	D3D12_BARRIER_ACCESS_SHADER_RESOURCE
 	| D3D12_BARRIER_ACCESS_COPY_DEST
@@ -656,14 +663,26 @@ TStaticArray<D3D12_BARRIER_ACCESS, 32> LayoutCompatibleAccess = {
 	//D3D12_BARRIER_LAYOUT_COMPUTE_QUEUE_COPY_DEST
 	D3D12_BARRIER_ACCESS_COPY_DEST,
 
+#if D3D12RHI_HAS_LAYOUT_VIDEO_QUEUE_COMMON
 	//D3D12_BARRIER_LAYOUT_VIDEO_QUEUE_COMMON,
 	D3D12_BARRIER_ACCESS_COPY_SOURCE
 	| D3D12_BARRIER_ACCESS_COPY_DEST,
+#else
+	// RESERVED
+	D3D12_BARRIER_ACCESS_NO_ACCESS,
+#endif
 
+#if D3D12RHI_COMMON_AND_PRESENT_ARE_DISTINCT
 	//D3D12_BARRIER_LAYOUT_PRESENT
 	D3D12_BARRIER_ACCESS_SHADER_RESOURCE
 	| D3D12_BARRIER_ACCESS_COPY_DEST
 	| D3D12_BARRIER_ACCESS_COPY_SOURCE,
+#endif
+
+#if D3D12RHI_HAS_LAYOUT_DIRECT_QUEUE_GENERIC_READ_COMPUTE_QUEUE_ACCESSIBLE
+	// D3D12_BARRIER_LAYOUT_DIRECT_QUEUE_GENERIC_READ_COMPUTE_QUEUE_ACCESSIBLE
+	D3D12_BARRIER_ACCESS_NO_ACCESS // @TODO - Fill this in after deducing from the debug layer!
+#endif
 };
 
 // These are not the same layout on some platforms but the compatible access should be the same
@@ -909,7 +928,7 @@ static bool AccessIsCompatibleWithLayout(
 		return true;
 	}
 
-	if (static_cast<int32>(InLayout) > LayoutCompatibleAccess.Num())
+	if (static_cast<int32>(InLayout) >= LayoutCompatibleAccess.Num())
 	{
 		return false;
 	}
diff --git a/Engine/Source/Runtime/D3D12RHI/Private/Windows/WindowsD3D12RHIDefinitions.h b/Engine/Source/Runtime/D3D12RHI/Private/Windows/WindowsD3D12RHIDefinitions.h
--- a/Engine/Source/Runtime/D3D12RHI/Private/Windows/WindowsD3D12RHIDefinitions.h
+++ b/Engine/Source/Runtime/D3D12RHI/Private/Windows/WindowsD3D12RHIDefinitions.h
@@ -39,6 +39,10 @@ typedef ID3D12GraphicsCommandList ID3D12CopyCommandList;
 #define D3D12RHI_ALLOW_ENHANCED_BARRIERS 1
 #define D3D12RHI_ALLOW_LEGACY_BARRIERS   1
 
+#define D3D12RHI_HAS_LAYOUT_VIDEO_QUEUE_COMMON                                  (UE_D3D12_SDK_VERSION  < 618)
+#define D3D12RHI_HAS_LAYOUT_DIRECT_QUEUE_GENERIC_READ_COMPUTE_QUEUE_ACCESSIBLE  (UE_D3D12_SDK_VERSION >= 618)
+#define D3D12RHI_COMMON_AND_PRESENT_ARE_DISTINCT                                0
+
 // Needed for XR
 #define D3D12RHI_SUPPORTS_ALIASED_RHI_RESOURCES 1
 
