Remove legacy GPU profiler code paths from D3D12 RHI

diff --git a/Engine/Source/Runtime/D3D12RHI/Private/D3D12Adapter.cpp b/Engine/Source/Runtime/D3D12RHI/Private/D3D12Adapter.cpp
--- a/Engine/Source/Runtime/D3D12RHI/Private/D3D12Adapter.cpp
+++ b/Engine/Source/Runtime/D3D12RHI/Private/D3D12Adapter.cpp
@@ -63,6 +63,16 @@ TAutoConsoleVariable<int32> GD3D12DebugCvar (
 	TEXT("4 to enable CONTINUING on errors (-d3dcontinueonerrors)\n"),
 	ECVF_RenderThreadSafe | ECVF_ReadOnly);
 
+TAutoConsoleVariable<int32> CVarD3D12RHIStablePowerState(
+	TEXT("D3D12.StablePowerState"),
+	0,
+	TEXT("Enable stable power state. This increases GPU timing measurement accuracy but may decrease overall GPU clock rate.\n")
+	TEXT("    0 (default): off\n")
+	TEXT("    1          : set during profiling (no longer implemented)\n")
+	TEXT("    2          : set on startup\n"),
+	ECVF_Default
+);
+
 bool D3D12_ShouldLogD3DDebugWarnings()
 {
 	return GD3D12DebugCvar.GetValueOnAnyThread() > 1;
@@ -795,8 +805,7 @@ void FD3D12Adapter::CreateRootDevice(bool bWithDebug)
 
 	if (RootDevice)
 	{
-		extern int32 GD3D12RHIStablePowerState;
-		if (GD3D12RHIStablePowerState == 2)
+		if (CVarD3D12RHIStablePowerState.GetValueOnAnyThread() == 2)
 		{
 			bool bWorked = SUCCEEDED(RootDevice->SetStablePowerState(true));
 			// This will fail if windows developper mode is not enabled. Windows will remove the adapter on failure, so we can't really gracefully exit here.
diff --git a/Engine/Source/Runtime/D3D12RHI/Private/D3D12Buffer.cpp b/Engine/Source/Runtime/D3D12RHI/Private/D3D12Buffer.cpp
--- a/Engine/Source/Runtime/D3D12RHI/Private/D3D12Buffer.cpp
+++ b/Engine/Source/Runtime/D3D12RHI/Private/D3D12Buffer.cpp
@@ -991,8 +991,4 @@ void FD3D12CommandContext::RHICopyBufferRegion(FRHIBuffer* DestBufferRHI, uint64
 	);
 
 	ConditionalSplitCommandList();
-
-#if (RHI_NEW_GPU_PROFILER == 0)
-	BufferDevice->RegisterGPUWork(1);
-#endif
 }
diff --git a/Engine/Source/Runtime/D3D12RHI/Private/D3D12CommandContext.cpp b/Engine/Source/Runtime/D3D12RHI/Private/D3D12CommandContext.cpp
--- a/Engine/Source/Runtime/D3D12RHI/Private/D3D12CommandContext.cpp
+++ b/Engine/Source/Runtime/D3D12RHI/Private/D3D12CommandContext.cpp
@@ -160,7 +160,6 @@ void FD3D12ContextCommon::BindDiagnosticBuffer(FD3D12RootSignature const* RootSi
 		#endif
 		}
 
-	#if RHI_NEW_GPU_PROFILER
 		{
 			FlushProfilerStats();
 
@@ -168,21 +167,10 @@ void FD3D12ContextCommon::BindDiagnosticBuffer(FD3D12RootSignature const* RootSi
 			FD3D12QueryLocation TimestampQuery = AllocateQuery(ED3D12QueryType::ProfilerTimestampTOP, &Event.GPUTimestampTOP);
 			EndQuery(TimestampQuery);
 		}
-	#else
-		if (IsDefaultContext() && !IsAsyncComputeContext())
-		{
-			FD3D12GPUProfiler& GPUProfiler = GetParentDevice()->GetGPUProfiler();
-			if (GPUProfiler.IsProfilingGPU())
-			{
-				GPUProfiler.PushEvent(GetNameStr(), FColor::White);
-			}
-		}
-	#endif // (RHI_NEW_GPU_PROFILER == 0)
 	}
 
 	void FD3D12CommandContext::RHIEndBreadcrumbGPU(FRHIBreadcrumbNode* Breadcrumb)
 	{
-	#if RHI_NEW_GPU_PROFILER
 		{
 			FlushProfilerStats();
 
@@ -190,16 +178,6 @@ void FD3D12ContextCommon::BindDiagnosticBuffer(FD3D12RootSignature const* RootSi
 			FD3D12QueryLocation TimestampQuery = AllocateQuery(ED3D12QueryType::ProfilerTimestampBOP, &Event.GPUTimestampBOP);
 			EndQuery(TimestampQuery);
 		}
-	#else
-		if (IsDefaultContext() && !IsAsyncComputeContext())
-		{
-			FD3D12GPUProfiler& GPUProfiler = GetParentDevice()->GetGPUProfiler();
-			if (GPUProfiler.IsProfilingGPU())
-			{
-				GPUProfiler.PopEvent();
-			}
-		}
-	#endif // (RHI_NEW_GPU_PROFILER == 0)
 
 		// Only emit formatted strings to platform APIs when requested.
 		if (ShouldEmitBreadcrumbs())
@@ -299,10 +277,8 @@ FD3D12QueryLocation FD3D12ContextCommon::AllocateQuery(ED3D12QueryType Type, voi
 
 	case ED3D12QueryType::TimestampRaw:
 	case ED3D12QueryType::TimestampMicroseconds:
-#if RHI_NEW_GPU_PROFILER
 	case ED3D12QueryType::ProfilerTimestampTOP:
 	case ED3D12QueryType::ProfilerTimestampBOP:
-#endif
 		return TimestampQueries.Allocate(Type, Target);
 
 	case ED3D12QueryType::Occlusion:
@@ -485,9 +461,7 @@ void FD3D12CommandContext::Finalize(TArray<FD3D12Payload*>& OutPayloads)
 	GetParentDevice()->GetBindlessDescriptorManager().FinalizeContext(*this);
 #endif
 
-#if RHI_NEW_GPU_PROFILER
 	FlushProfilerStats();
-#endif
 
 	FD3D12ContextCommon::Finalize(OutPayloads);
 }
diff --git a/Engine/Source/Runtime/D3D12RHI/Private/D3D12CommandContext.h b/Engine/Source/Runtime/D3D12RHI/Private/D3D12CommandContext.h
--- a/Engine/Source/Runtime/D3D12RHI/Private/D3D12CommandContext.h
+++ b/Engine/Source/Runtime/D3D12RHI/Private/D3D12CommandContext.h
@@ -544,7 +544,6 @@ public:
 	void ClearShaderResources(FD3D12BaseShaderResource* Resource, EShaderParameterTypeMask ShaderParameterTypeMask);
 	void ClearAllShaderResources();
 
-#if RHI_NEW_GPU_PROFILER
 	void FlushProfilerStats()
 	{
 		// Flush accumulated draw stats
@@ -554,7 +553,6 @@ public:
 			StatEvent = {};
 		}
 	}
-#endif
 
 	FD3D12FastConstantAllocator ConstantsAllocator;
 
@@ -667,9 +665,6 @@ public:
 	virtual void RHICopyToStagingBuffer(FRHIBuffer* SourceBuffer, FRHIStagingBuffer* DestinationStagingBuffer, uint32 Offset, uint32 NumBytes) final override;
 	virtual void RHIBeginRenderQuery(FRHIRenderQuery* RenderQuery) final override;
 	virtual void RHIEndRenderQuery(FRHIRenderQuery* RenderQuery) final override;
-#if (RHI_NEW_GPU_PROFILER == 0)
-	virtual void RHICalibrateTimers(FRHITimestampCalibrationQuery* CalibrationQuery) final override;
-#endif
 	virtual void RHISetStreamSource(uint32 StreamIndex, FRHIBuffer* VertexBuffer, uint32 Offset) final override;
 	virtual void RHISetViewport(float MinX, float MinY, float MinZ, float MaxX, float MaxY, float MaxZ) final override;
 	virtual void RHISetStereoViewport(float LeftMinX, float RightMinX, float LeftMinY, float RightMinY, float MinZ, float LeftMaxX, float RightMaxX, float LeftMaxY, float RightMaxY, float MaxZ) override;
@@ -982,12 +977,6 @@ public:
 	{
 		ContextRedirect(RHIEndRenderQuery(RenderQuery));
 	}
-#if (RHI_NEW_GPU_PROFILER == 0)
-	FORCEINLINE virtual void RHICalibrateTimers(FRHITimestampCalibrationQuery* CalibrationQuery) final override
-	{
-		ContextRedirect(RHICalibrateTimers(CalibrationQuery));
-	}
-#endif
 	FORCEINLINE virtual void RHISetStreamSource(uint32 StreamIndex, FRHIBuffer* VertexBuffer, uint32 Offset) final override
 	{
 		ContextRedirect(RHISetStreamSource(StreamIndex, VertexBuffer, Offset));
diff --git a/Engine/Source/Runtime/D3D12RHI/Private/D3D12CommandList.cpp b/Engine/Source/Runtime/D3D12RHI/Private/D3D12CommandList.cpp
--- a/Engine/Source/Runtime/D3D12RHI/Private/D3D12CommandList.cpp
+++ b/Engine/Source/Runtime/D3D12RHI/Private/D3D12CommandList.cpp
@@ -255,11 +255,9 @@ void FD3D12CommandList::BeginLocalQueries()
 
 	if (State.BeginTimestamp)
 	{
-#if RHI_NEW_GPU_PROFILER
 		// CPUTimestamp is filled in at submission time in FlushProfilerEvents
 		auto& Event = EmplaceProfilerEvent<UE::RHI::GPUProfiler::FEvent::FBeginWork>(0);
 		State.BeginTimestamp.Target = &Event.GPUTimestampTOP;
-#endif
 
 		EndQuery(State.BeginTimestamp);
 	}
@@ -284,10 +282,8 @@ void FD3D12CommandList::EndLocalQueries()
 
 	if (State.EndTimestamp)
 	{
-#if RHI_NEW_GPU_PROFILER
 		auto& Event = EmplaceProfilerEvent<UE::RHI::GPUProfiler::FEvent::FEndWork>();
 		State.EndTimestamp.Target = &Event.GPUTimestampBOP;
-#endif
 
 		EndQuery(State.EndTimestamp);
 	}
@@ -342,37 +338,20 @@ void FD3D12CommandList::EndQuery(FD3D12QueryLocation const& Location)
 				Position = ED3D12QueryPosition::BottomOfPipe;
 				break;
 
-#if RHI_NEW_GPU_PROFILER
 			case ED3D12QueryType::ProfilerTimestampTOP:
-#else
-			case ED3D12QueryType::CommandListBegin:
-			case ED3D12QueryType::IdleBegin:
-#endif
 				Position = ED3D12QueryPosition::TopOfPipe;
 				break;
 
 			case ED3D12QueryType::TimestampMicroseconds:
 			case ED3D12QueryType::TimestampRaw:
-#if RHI_NEW_GPU_PROFILER
 			case ED3D12QueryType::ProfilerTimestampBOP:
-#else
-			case ED3D12QueryType::CommandListEnd:
-			case ED3D12QueryType::IdleEnd:
-#endif
 				Position = ED3D12QueryPosition::BottomOfPipe;
 				break;
 			}
 
 			WriteTimestamp(Location, Position);
-
-#if RHI_NEW_GPU_PROFILER == 0
-			// Command list begin/end timestamps are handled separately by the 
-			// submission thread, so shouldn't be in the TimestampQueries array.
-			if (Location.Type != ED3D12QueryType::CommandListBegin && Location.Type != ED3D12QueryType::CommandListEnd)
-#endif
-			{
-				State.TimestampQueries.Add(Location);
-			}
+			
+			State.TimestampQueries.Add(Location);
 		}
 		break;
 	}
@@ -391,19 +370,12 @@ void FD3D12CommandList::WriteTimestamp(FD3D12QueryLocation const& Location, ED3D
 
 FD3D12CommandList::FState::FState(FD3D12CommandAllocator* CommandAllocator, FD3D12QueryAllocator* TimestampAllocator, FD3D12QueryAllocator* PipelineStatsAllocator)
 	: CommandAllocator(CommandAllocator)
-#if RHI_NEW_GPU_PROFILER
 	, EventStream(CommandAllocator->Device->GetQueue(CommandAllocator->QueueType).GetProfilerQueue())
-#endif
 {
 	if (TimestampAllocator)
 	{
-#if RHI_NEW_GPU_PROFILER
 		BeginTimestamp = TimestampAllocator->Allocate(ED3D12QueryType::ProfilerTimestampTOP, nullptr);
 		EndTimestamp   = TimestampAllocator->Allocate(ED3D12QueryType::ProfilerTimestampBOP, nullptr);
-#else
-		BeginTimestamp = TimestampAllocator->Allocate(ED3D12QueryType::CommandListBegin, nullptr);
-		EndTimestamp   = TimestampAllocator->Allocate(ED3D12QueryType::CommandListEnd  , nullptr);
-#endif
 	}
 
 	if (PipelineStatsAllocator)
diff --git a/Engine/Source/Runtime/D3D12RHI/Private/D3D12CommandList.h b/Engine/Source/Runtime/D3D12RHI/Private/D3D12CommandList.h
--- a/Engine/Source/Runtime/D3D12RHI/Private/D3D12CommandList.h
+++ b/Engine/Source/Runtime/D3D12RHI/Private/D3D12CommandList.h
@@ -109,7 +109,6 @@ public:
 	FD3D12ResidencySet* CloseResidencySet();
 #endif // ENABLE_RESIDENCY_MANAGEMENT
 
-#if RHI_NEW_GPU_PROFILER
 	template <typename TEventType, typename... TArgs>
 	TEventType& EmplaceProfilerEvent(TArgs&&... Args)
 	{
@@ -133,7 +132,6 @@ public:
 
 		Destination.Append(MoveTemp(State.EventStream));
 	}
-#endif
 
 private:
 	struct FInterfaces
@@ -312,10 +310,8 @@ private:
 		bool bLocalQueriesEnded = false;
 	#endif
 
-	#if RHI_NEW_GPU_PROFILER
 		UE::RHI::GPUProfiler::FEventStream EventStream;
 		TArray<UE::RHI::GPUProfiler::FEvent::FBeginWork*, TInlineAllocator<8>> BeginEvents;
-	#endif
 		
 	} State;
 };
\ No newline at end of file
diff --git a/Engine/Source/Runtime/D3D12RHI/Private/D3D12Commands.cpp b/Engine/Source/Runtime/D3D12RHI/Private/D3D12Commands.cpp
--- a/Engine/Source/Runtime/D3D12RHI/Private/D3D12Commands.cpp
+++ b/Engine/Source/Runtime/D3D12RHI/Private/D3D12Commands.cpp
@@ -96,13 +96,6 @@ void FD3D12CommandContext::RHISetStreamSource(uint32 StreamIndex, FRHIBuffer* Ve
 
 void FD3D12CommandContext::SetupDispatch(uint32 ThreadGroupCountX, uint32 ThreadGroupCountY, uint32 ThreadGroupCountZ)
 {
-#if (RHI_NEW_GPU_PROFILER == 0)
-	if (IsDefaultContext())
-	{
-		GetParentDevice()->RegisterGPUDispatch(FIntVector(ThreadGroupCountX, ThreadGroupCountY, ThreadGroupCountZ));
-	}
-#endif
-
 	FlushPendingDescriptorUpdates();
 
 	CommitComputeShaderConstants();
@@ -1003,15 +996,6 @@ void FD3D12CommandContext::SetRenderTargetsAndClear(const FRHISetRenderTargetsIn
 #endif
 }
 
-#if (RHI_NEW_GPU_PROFILER == 0)
-void FD3D12CommandContext::RHICalibrateTimers(FRHITimestampCalibrationQuery* CalibrationQuery)
-{
-	FGPUTimingCalibrationTimestamp Timestamp = GetParentDevice()->GetCalibrationTimestamp(QueueType);
-	CalibrationQuery->CPUMicroseconds[GetGPUIndex()] = Timestamp.CPUMicroseconds;
-	CalibrationQuery->GPUMicroseconds[GetGPUIndex()] = Timestamp.GPUMicroseconds;
-}
-#endif
-
 // Primitive drawing.
 
 void FD3D12CommandContext::CommitNonComputeShaderConstants()
@@ -1173,13 +1157,6 @@ void FD3D12CommandContext::RHIDispatchGraphicsShaderBundle(
 
 void FD3D12CommandContext::SetupDraw(FRHIBuffer* IndexBufferRHI, uint32 NumPrimitives /* = 0 */, uint32 NumVertices /* = 0 */)
 {
-#if (RHI_NEW_GPU_PROFILER == 0)
-	if (IsDefaultContext() && Device->GetGPUProfiler().bTrackingEvents)
-	{
-		GetParentDevice()->RegisterGPUWork(NumPrimitives, NumVertices);
-	}
-#endif
-
 	FlushPendingDescriptorUpdates();
 
 	CommitGraphicsResourceTables();
@@ -1200,13 +1177,6 @@ void FD3D12CommandContext::SetupDraw(FRHIBuffer* IndexBufferRHI, uint32 NumPrimi
 
 void FD3D12CommandContext::SetupDispatchDraw(uint32 ThreadGroupCountX, uint32 ThreadGroupCountY, uint32 ThreadGroupCountZ)
 {
-#if (RHI_NEW_GPU_PROFILER == 0)
-	if (IsDefaultContext() && Device->GetGPUProfiler().bTrackingEvents)
-	{
-		GetParentDevice()->RegisterGPUDispatch(FIntVector(ThreadGroupCountX, ThreadGroupCountY, ThreadGroupCountZ));
-	}
-#endif
-
 	FlushPendingDescriptorUpdates();
 
 	CommitGraphicsResourceTables();
@@ -1472,13 +1442,6 @@ void FD3D12CommandContext::RHIClearMRTImpl(bool* bClearColorArray, int32 NumClea
 		ConditionalSplitCommandList();
 	}
 
-#if (RHI_NEW_GPU_PROFILER == 0)
-	if (IsDefaultContext())
-	{
-		GetParentDevice()->RegisterGPUWork(0);
-	}
-#endif
-
 	DEBUG_EXECUTE_COMMAND_LIST(this);
 }
 
diff --git a/Engine/Source/Runtime/D3D12RHI/Private/D3D12Device.cpp b/Engine/Source/Runtime/D3D12RHI/Private/D3D12Device.cpp
--- a/Engine/Source/Runtime/D3D12RHI/Private/D3D12Device.cpp
+++ b/Engine/Source/Runtime/D3D12RHI/Private/D3D12Device.cpp
@@ -35,9 +35,7 @@ static uint32 GetQueryHeapPoolIndex(D3D12_QUERY_HEAP_TYPE HeapType)
 
 FD3D12Timing::FD3D12Timing(FD3D12Queue& Queue)
 	: Queue(Queue)
-#if RHI_NEW_GPU_PROFILER
 	, EventStream(Queue.GetProfilerQueue())
-#endif
 {}
 
 FD3D12Queue::FD3D12Queue(FD3D12Device* Device, ED3D12QueueType QueueType, int32 QueueIndex)
@@ -84,7 +82,6 @@ FD3D12Queue::~FD3D12Queue()
 	check(PendingInterrupt.IsEmpty());
 }
 
-#if RHI_NEW_GPU_PROFILER
 UE::RHI::GPUProfiler::FQueue FD3D12Queue::GetProfilerQueue() const
 {
 	UE::RHI::GPUProfiler::FQueue Queue;
@@ -101,15 +98,11 @@ UE::RHI::GPUProfiler::FQueue FD3D12Queue::GetProfilerQueue() const
 
 	return Queue;
 }
-#endif // RHI_NEW_GPU_PROFILER
 
 FD3D12Device::FD3D12Device(FRHIGPUMask InGPUMask, FD3D12Adapter* InAdapter)
 	: FD3D12SingleNodeGPUObject(InGPUMask)
 	, FD3D12AdapterChild       (InAdapter)
 	, TileMappingFence         (nullptr)
-#if (RHI_NEW_GPU_PROFILER == 0)
-	, GPUProfilingData         (this)
-#endif
 	, ResidencyManager         (*this)
 	, DescriptorHeapManager    (this)
 #if PLATFORM_SUPPORTS_BINDLESS_RENDERING
@@ -810,30 +803,6 @@ uint64 FD3D12Device::GetTimestampFrequency(ED3D12QueueType QueueType)
 	return Frequency;
 }
 
-#if (RHI_NEW_GPU_PROFILER == 0)
-FGPUTimingCalibrationTimestamp FD3D12Device::GetCalibrationTimestamp(ED3D12QueueType QueueType)
-{
-	check(QueueType != ED3D12QueueType::Copy || GetParentAdapter()->AreCopyQueueTimestampQueriesSupported());
-
-	TRACE_CPUPROFILER_EVENT_SCOPE(D3D12GetCalibrationTimestamp);
-
-	uint64 GPUTimestampFrequency = GetTimestampFrequency(QueueType);
-
-	LARGE_INTEGER CPUTimestempFrequency;
-	QueryPerformanceFrequency(&CPUTimestempFrequency);
-
-	uint64 GPUTimestamp, CPUTimestamp;
-	VERIFYD3D12RESULT(Queues[(uint32)QueueType].D3DCommandQueue->GetClockCalibration(&GPUTimestamp, &CPUTimestamp));
-
-	FGPUTimingCalibrationTimestamp Result = {};
-
-	Result.GPUMicroseconds = uint64(GPUTimestamp * (1e6 / GPUTimestampFrequency));
-	Result.CPUMicroseconds = uint64(CPUTimestamp * (1e6 / CPUTimestempFrequency.QuadPart));
-
-	return Result;
-}
-#endif // (RHI_NEW_GPU_PROFILER == 0)
-
 void FD3D12Device::InitExplicitDescriptorHeap()
 {
 	check(ExplicitDescriptorHeapCache == nullptr);
diff --git a/Engine/Source/Runtime/D3D12RHI/Private/D3D12Device.h b/Engine/Source/Runtime/D3D12RHI/Private/D3D12Device.h
--- a/Engine/Source/Runtime/D3D12RHI/Private/D3D12Device.h
+++ b/Engine/Source/Runtime/D3D12RHI/Private/D3D12Device.h
@@ -16,7 +16,6 @@ D3D12Device.h: D3D12 Device Interfaces
 #include "D3D12Queue.h"
 #include "D3D12Resources.h"
 #include "D3D12Submission.h"
-#include "D3D12GPUProfiler.h"
 
 #include "Containers/LruCache.h"
 #include "Containers/SpscQueue.h"
@@ -72,29 +71,12 @@ public:
 	FD3D12Queue& Queue;
 	D3D12_QUERY_DATA_PIPELINE_STATISTICS PipelineStats {};
 
-#if RHI_NEW_GPU_PROFILER
 	// Timer calibration data
 	uint64 GPUFrequency = 0, GPUTimestamp = 0;
 	uint64 CPUFrequency = 0, CPUTimestamp = 0;
 
 	UE::RHI::GPUProfiler::FEventStream EventStream;
 
-#else
-
-	TArray<uint64> Timestamps;
-	int32 TimestampIndex = 0;
-	uint64 BusyCycles = 0;
-
-	uint64 GetCurrentTimestamp()  const { return Timestamps[TimestampIndex]; }
-	uint64 GetPreviousTimestamp() const { return Timestamps[TimestampIndex - 1]; }
-
-	bool HasMoreTimestamps() const { return TimestampIndex < Timestamps.Num(); }
-	bool IsStartingWork()    const { return (TimestampIndex & 0x01) == 0x00; }
-
-	void AdvanceTimestamp() { TimestampIndex++; }
-
-#endif
-
 	FD3D12Timing(FD3D12Queue& Queue);
 };
 
@@ -122,8 +104,6 @@ public:
 	FD3D12CommandAllocator* BarrierAllocator = nullptr;
 	FD3D12QueryAllocator    BarrierTimestamps;
 
-	uint32 NumCommandListsInBatch = 0;
-
 	FD3D12BatchedPayloadObjects BatchedObjects;
 
 	// A pool of reusable command list/allocator/context objects
@@ -162,9 +142,7 @@ public:
 	FD3D12Queue(FD3D12Device* Device, ED3D12QueueType QueueType, int32 QueueIndex);
 	~FD3D12Queue();
 
-#if RHI_NEW_GPU_PROFILER
 	UE::RHI::GPUProfiler::FQueue GetProfilerQueue() const;
-#endif
 
 private:
 	// Internal fence which may be used before calling ExecuteCommandLists
@@ -180,20 +158,8 @@ public:
 
 	ID3D12Device* GetDevice();
 
-#if (RHI_NEW_GPU_PROFILER == 0)
-	void RegisterGPUWork(uint32 NumPrimitives = 0, uint32 NumVertices = 0)	{ GPUProfilingData.RegisterGPUWork(NumPrimitives, NumVertices); }
-	void RegisterGPUDispatch(FIntVector GroupCount)	                        { GPUProfilingData.RegisterGPUDispatch(GroupCount); }
-
-	// GPU Profiler
-	FORCEINLINE FD3D12GPUProfiler& GetGPUProfiler() { return GPUProfilingData; }
-#endif
-
 	uint64 GetTimestampFrequency(ED3D12QueueType QueueType);
 
-#if (RHI_NEW_GPU_PROFILER == 0)
-	FGPUTimingCalibrationTimestamp GetCalibrationTimestamp(ED3D12QueueType QueueType);
-#endif
-
 	// Misc
 	void BlockUntilIdle();
 	D3D12_RESOURCE_ALLOCATION_INFO GetResourceAllocationInfoUncached(const FD3D12ResourceDesc& InDesc);
@@ -321,10 +287,6 @@ private:
 	void UpdateMSAASettings();
 	void UpdateConstantBufferPageProperties();
 
-#if (RHI_NEW_GPU_PROFILER == 0)
-	FD3D12GPUProfiler GPUProfilingData;
-#endif
-
 	struct FResidencyManager : public FD3D12ResidencyManager
 	{
 		FResidencyManager(FD3D12Device& Parent);
diff --git a/Engine/Source/Runtime/D3D12RHI/Private/D3D12EnhancedBarriers.cpp b/Engine/Source/Runtime/D3D12RHI/Private/D3D12EnhancedBarriers.cpp
--- a/Engine/Source/Runtime/D3D12RHI/Private/D3D12EnhancedBarriers.cpp
+++ b/Engine/Source/Runtime/D3D12RHI/Private/D3D12EnhancedBarriers.cpp
@@ -2119,7 +2119,6 @@ void FD3D12EnhancedBarriersBatcher::FlushIntoCommandList(
 {
 	auto InsertTimestamp = [&](bool bBegin)
 	{
-#if RHI_NEW_GPU_PROFILER
 		if (bBegin)
 		{
 			auto& Event = CommandList.EmplaceProfilerEvent<UE::RHI::GPUProfiler::FEvent::FEndWork>();
@@ -2131,13 +2130,6 @@ void FD3D12EnhancedBarriersBatcher::FlushIntoCommandList(
 			auto& Event = CommandList.EmplaceProfilerEvent<UE::RHI::GPUProfiler::FEvent::FBeginWork>(0);
 			CommandList.EndQuery(TimestampAllocator.Allocate(ED3D12QueryType::ProfilerTimestampTOP, &Event.GPUTimestampTOP));
 		}
-#else
-		ED3D12QueryType Type = bBegin
-			? ED3D12QueryType::IdleBegin
-			: ED3D12QueryType::IdleEnd;
-
-		CommandList.EndQuery(TimestampAllocator.Allocate(Type, nullptr));
-#endif
 	};
 
 	for (int32 BatchStart = 0, BatchEnd = 0; BatchStart < BarrierGroups.Num(); BatchStart = BatchEnd)
diff --git a/Engine/Source/Runtime/D3D12RHI/Private/D3D12GPUProfiler.cpp b/Engine/Source/Runtime/D3D12RHI/Private/D3D12GPUProfiler.cpp
deleted file mode 100644
diff --git a/Engine/Source/Runtime/D3D12RHI/Private/D3D12GPUProfiler.h b/Engine/Source/Runtime/D3D12RHI/Private/D3D12GPUProfiler.h
deleted file mode 100644
diff --git a/Engine/Source/Runtime/D3D12RHI/Private/D3D12LegacyBarriers.cpp b/Engine/Source/Runtime/D3D12RHI/Private/D3D12LegacyBarriers.cpp
--- a/Engine/Source/Runtime/D3D12RHI/Private/D3D12LegacyBarriers.cpp
+++ b/Engine/Source/Runtime/D3D12RHI/Private/D3D12LegacyBarriers.cpp
@@ -872,7 +872,6 @@ void FD3D12LegacyBarriersBatcher::FlushIntoCommandList(
 {
 	auto InsertTimestamp = [&](bool bBegin)
 	{
-#if RHI_NEW_GPU_PROFILER
 		if (bBegin)
 		{
 			auto& Event = CommandList.EmplaceProfilerEvent<UE::RHI::GPUProfiler::FEvent::FEndWork>();
@@ -884,13 +883,6 @@ void FD3D12LegacyBarriersBatcher::FlushIntoCommandList(
 			auto& Event = CommandList.EmplaceProfilerEvent<UE::RHI::GPUProfiler::FEvent::FBeginWork>(0);
 			CommandList.EndQuery(TimestampAllocator.Allocate(ED3D12QueryType::ProfilerTimestampTOP, &Event.GPUTimestampTOP));
 		}
-#else
-		ED3D12QueryType Type = bBegin
-			? ED3D12QueryType::IdleBegin
-			: ED3D12QueryType::IdleEnd;
-
-		CommandList.EndQuery(TimestampAllocator.Allocate(Type, nullptr));
-#endif
 	};
 
 	for (int32 BatchStart = 0, BatchEnd = 0; BatchStart < Barriers.Num(); BatchStart = BatchEnd)
diff --git a/Engine/Source/Runtime/D3D12RHI/Private/D3D12Query.cpp b/Engine/Source/Runtime/D3D12RHI/Private/D3D12Query.cpp
--- a/Engine/Source/Runtime/D3D12RHI/Private/D3D12Query.cpp
+++ b/Engine/Source/Runtime/D3D12RHI/Private/D3D12Query.cpp
@@ -344,8 +344,3 @@ bool FD3D12DynamicRHI::RHIGetRenderQueryResult(FRHIRenderQuery* QueryRHI, uint64
 	OutResult = *Query->Result;
 	return true;
 }
-
-/*=============================================================================
- * class FD3D12BufferedGPUTiming
- *=============================================================================*/
-
diff --git a/Engine/Source/Runtime/D3D12RHI/Private/D3D12Query.h b/Engine/Source/Runtime/D3D12RHI/Private/D3D12Query.h
--- a/Engine/Source/Runtime/D3D12RHI/Private/D3D12Query.h
+++ b/Engine/Source/Runtime/D3D12RHI/Private/D3D12Query.h
@@ -17,7 +17,6 @@ enum class ED3D12QueryType
 	PipelineStats,
 	TimestampMicroseconds,
 	TimestampRaw,
-#if RHI_NEW_GPU_PROFILER
 
 	// Top-of-pipe GPU timestamp, converted to the CPU clock domain (i.e. relative to FPlatformTime::Cycles64())
 	ProfilerTimestampTOP,
@@ -25,12 +24,6 @@ enum class ED3D12QueryType
 	// Bottom-of-pipe GPU timestamp, converted to the CPU clock domain (i.e. relative to FPlatformTime::Cycles64())
 	ProfilerTimestampBOP,
 
-#else
-	CommandListBegin,
-	CommandListEnd,
-	IdleBegin,
-	IdleEnd,
-#endif
 	Occlusion
 };
 
diff --git a/Engine/Source/Runtime/D3D12RHI/Private/D3D12RHI.cpp b/Engine/Source/Runtime/D3D12RHI/Private/D3D12RHI.cpp
--- a/Engine/Source/Runtime/D3D12RHI/Private/D3D12RHI.cpp
+++ b/Engine/Source/Runtime/D3D12RHI/Private/D3D12RHI.cpp
@@ -522,21 +522,6 @@ void FD3D12DynamicRHI::RHIProcessDeleteQueue()
 
 void FD3D12DynamicRHI::RHIEndFrame_RenderThread(FRHICommandListImmediate& RHICmdList)
 {
-	// Close the GPU profiler frame
-#if (RHI_NEW_GPU_PROFILER == 0)
-	RHICmdList.EnqueueLambdaMultiPipe(GetEnabledRHIPipelines(), FRHICommandListBase::EThreadFence::Enabled, TEXT("D3D12 EndFrame"),
-		[this](FD3D12ContextArray const& Contexts)
-	{
-		for (auto& Adapter : ChosenAdapters)
-		{
-			for (auto& Device : Adapter->GetDevices())
-			{
-				Device->GetGPUProfiler().EndFrame();
-			}
-		}
-	});
-#endif
-
 	for (auto& Adapter : ChosenAdapters)
 	{
 		Adapter->GetFrameFence().AdvanceTOP();
@@ -556,14 +541,10 @@ void FD3D12DynamicRHI::RHIEndFrame_RenderThread(FRHICommandListImmediate& RHICmd
 				Device->GetDefaultBufferAllocator().BeginFrame(Contexts);
 				Device->GetTextureAllocator().BeginFrame(Contexts);
 
-#if D3D12_RHI_RAYTRACING
+			#if D3D12_RHI_RAYTRACING
 				// @todo dev-pr - explicit use of graphics context - nothing is synchronizing async compute - needs refactor
 				Device->GetRayTracingCompactionRequestHandler()->Update(*Contexts[ERHIPipeline::Graphics]->GetSingleDeviceContext(Device->GetGPUIndex()));
-#endif // D3D12_RHI_RAYTRACING
-
-#if (RHI_NEW_GPU_PROFILER == 0)
-				Device->GetGPUProfiler().BeginFrame();
-#endif
+			#endif
 			}
 		}
 	});
@@ -607,7 +588,6 @@ void FD3D12DynamicRHI::RHIEndFrame(const FRHIEndFrameArgs& Args)
 
 		Payload.Timing = CurrentTimingPerQueue.CreateNew(Payload.Queue);
 
-	#if RHI_NEW_GPU_PROFILER
 		ERHIPipeline Pipeline;
 		switch (Payload.Queue.QueueType)
 		{
@@ -630,7 +610,6 @@ void FD3D12DynamicRHI::RHIEndFrame(const FRHIEndFrameArgs& Args)
 			, Args.StatsFrame
 		#endif
 		);
-	#endif
 	});
 
 	// Pump the interrupt queue to gather completed events
diff --git a/Engine/Source/Runtime/D3D12RHI/Private/D3D12RHIPrivate.h b/Engine/Source/Runtime/D3D12RHI/Private/D3D12RHIPrivate.h
--- a/Engine/Source/Runtime/D3D12RHI/Private/D3D12RHIPrivate.h
+++ b/Engine/Source/Runtime/D3D12RHI/Private/D3D12RHIPrivate.h
@@ -319,12 +319,6 @@ public:
 	virtual FD3D12CommandContext* CreateCommandContext(FD3D12Device* InParent, ED3D12QueueType InQueueType, bool InIsDefaultContext);
 	virtual void CreateCommandQueue(FD3D12Device* Device, const D3D12_COMMAND_QUEUE_DESC& Desc, TRefCountPtr<ID3D12CommandQueue>& OutCommandQueue);
 
-	virtual bool GetHardwareGPUFrameTime(double& OutGPUFrameTime) const
-	{
-		OutGPUFrameTime = 0.0;
-		return false;
-	}
-
 	virtual void RHIEndFrame_RenderThread(FRHICommandListImmediate& RHICmdList) final override;
 	virtual void RHIEndFrame(const FRHIEndFrameArgs& Args) final override;
 
diff --git a/Engine/Source/Runtime/D3D12RHI/Private/D3D12RayTracing.cpp b/Engine/Source/Runtime/D3D12RHI/Private/D3D12RayTracing.cpp
--- a/Engine/Source/Runtime/D3D12RHI/Private/D3D12RayTracing.cpp
+++ b/Engine/Source/Runtime/D3D12RHI/Private/D3D12RayTracing.cpp
@@ -5751,13 +5751,6 @@ static void DispatchRays(FD3D12CommandContext& CommandContext,
 		{
 			CommandContext.RayTracingCommandList()->DispatchRays(&DispatchDesc);
 		}
-
-	#if (RHI_NEW_GPU_PROFILER == 0)
-		if (CommandContext.IsDefaultContext())
-		{
-			CommandContext.GetParentDevice()->RegisterGPUWork(1);
-		}
-	#endif
 	}
 
 	// Restore old global descriptor heaps
diff --git a/Engine/Source/Runtime/D3D12RHI/Private/D3D12RenderTarget.cpp b/Engine/Source/Runtime/D3D12RHI/Private/D3D12RenderTarget.cpp
--- a/Engine/Source/Runtime/D3D12RHI/Private/D3D12RenderTarget.cpp
+++ b/Engine/Source/Runtime/D3D12RHI/Private/D3D12RenderTarget.cpp
@@ -184,13 +184,6 @@ void FD3D12CommandContext::ResolveTextureUsingShader(
 		// Clear the destination texture.
 		if (bClearDestTexture)
 		{
-		#if (RHI_NEW_GPU_PROFILER == 0)
-			if (IsDefaultContext())
-			{
-				GetParentDevice()->RegisterGPUWork(0);
-			}
-		#endif
-
 			FlushResourceBarriers();
 
 			GraphicsCommandList()->ClearDepthStencilView(DestTextureDSV->GetOfflineCpuHandle(), D3D12_CLEAR_FLAG_DEPTH | D3D12_CLEAR_FLAG_STENCIL, 0, 0, 0, nullptr);
@@ -215,13 +208,6 @@ void FD3D12CommandContext::ResolveTextureUsingShader(
 		// Clear the destination texture.
 		if (bClearDestTexture)
 		{
-		#if (RHI_NEW_GPU_PROFILER == 0)
-			if (IsDefaultContext())
-			{
-				GetParentDevice()->RegisterGPUWork(0);
-			}
-		#endif
-
 			FlushResourceBarriers();
 
 			FLinearColor ClearColor(0, 0, 0, 0);
@@ -310,13 +296,6 @@ void FD3D12CommandContext::ResolveTexture(UE::RHICore::FResolveTextureInfo Info)
 {
 	uint32 GPUIndex = GetGPUIndex();
 
-#if (RHI_NEW_GPU_PROFILER == 0)
-	if (IsDefaultContext())
-	{
-		GetParentDevice()->RegisterGPUWork();
-	}
-#endif
-
 	FD3D12Texture* SourceTexture         = GetD3D12TextureFromRHITexture(Info.SourceTexture, GPUIndex);
 	FD3D12Resource* SourceResource       = SourceTexture->GetResource();
 	const FRHITextureDesc& SourceDesc    = SourceTexture->GetDesc();
diff --git a/Engine/Source/Runtime/D3D12RHI/Private/D3D12Submission.cpp b/Engine/Source/Runtime/D3D12RHI/Private/D3D12Submission.cpp
--- a/Engine/Source/Runtime/D3D12RHI/Private/D3D12Submission.cpp
+++ b/Engine/Source/Runtime/D3D12RHI/Private/D3D12Submission.cpp
@@ -27,12 +27,6 @@ static TAutoConsoleVariable<int32> CVarRHIUseSubmissionThread(
 
 DECLARE_CYCLE_STAT(TEXT("Submit"), STAT_D3D12Submit, STATGROUP_D3D12RHI);
 
-DECLARE_CYCLE_STAT(TEXT("GPU Total Time [All Queues]"), STAT_RHI_GPUTotalTime, STATGROUP_D3D12RHI);
-DECLARE_CYCLE_STAT(TEXT("GPU Total Time [Hardware Timer]"), STAT_RHI_GPUTotalTimeHW, STATGROUP_D3D12RHI);
-DECLARE_CYCLE_STAT(TEXT("GPU Total Time [Graphics]"), STAT_RHI_GPUTotalTimeGraphics, STATGROUP_D3D12RHI);
-DECLARE_CYCLE_STAT(TEXT("GPU Total Time [Async Compute]"), STAT_RHI_GPUTotalTimeAsyncCompute, STATGROUP_D3D12RHI);
-DECLARE_CYCLE_STAT(TEXT("GPU Total Time [Copy]"), STAT_RHI_GPUTotalTimeCopy, STATGROUP_D3D12RHI);
-
 DECLARE_STATS_GROUP(TEXT("D3D12RHIPipeline"), STATGROUP_D3D12RHIPipeline, STATCAT_Advanced);
 
 DECLARE_DWORD_ACCUMULATOR_STAT(TEXT("GPU IA Vertices"   ), STAT_D3D12RHI_IAVertices   , STATGROUP_D3D12RHIPipeline);
@@ -179,23 +173,17 @@ void FD3D12DynamicRHI::InitializeSubmissionPipe()
 	// Initialize the timing structs in each queue, and the engine GPU profilers
 	{
 		TArray<FD3D12Payload*> Payloads;
-	#if RHI_NEW_GPU_PROFILER
 		TArray<UE::RHI::GPUProfiler::FQueue> ProfilerQueues;
-	#endif
 
 		ForEachQueue([&](FD3D12Queue& Queue)
 		{
 			FD3D12Payload* Payload = Payloads.Emplace_GetRef(new FD3D12Payload(Queue));
 			Payload->Timing = CurrentTimingPerQueue.CreateNew(Queue);
 
-		#if RHI_NEW_GPU_PROFILER
 			ProfilerQueues.Add(Queue.GetProfilerQueue());
-		#endif
 		});
 
-	#if RHI_NEW_GPU_PROFILER
 		UE::RHI::GPUProfiler::InitializeQueues(ProfilerQueues);
-	#endif
 		SubmitPayloads(MoveTemp(Payloads));
 	}
 }
@@ -267,8 +255,6 @@ void FD3D12DynamicRHI::SubmitCommands(TConstArrayView<FD3D12FinalizedCommands*>
 {
 	SCOPED_NAMED_EVENT_TEXT("CommandList_Submit", FColor::Magenta);
 
-#if RHI_NEW_GPU_PROFILER
-
 	TArray<FD3D12Payload*> AllPayloads;
 	for (FD3D12FinalizedCommands* Payloads : Commands)
 	{
@@ -295,38 +281,6 @@ void FD3D12DynamicRHI::SubmitCommands(TConstArrayView<FD3D12FinalizedCommands*>
 	}
 
 	SubmitPayloads(MoveTemp(AllPayloads));
-
-#else
-
-	TArray<FD3D12Payload*> AllPayloads;
-	#if WITH_RHI_BREADCRUMBS
-	TArray<TSharedPtr<FRHIBreadcrumbAllocator>> BreadcrumbAllocators;
-	#endif
-
-	for (FD3D12FinalizedCommands* Payloads : Commands)
-	{
-	#if WITH_RHI_BREADCRUMBS
-		for (FD3D12Payload* Payload : *Payloads)
-		{
-			Payload->BreadcrumbRange = Payloads->BreadcrumbRange;
-		}
-	#endif
-
-		AllPayloads.Append(MoveTemp(static_cast<TArray<FD3D12Payload*>&>(*Payloads)));
-	#if WITH_RHI_BREADCRUMBS
-		BreadcrumbAllocators.Append(MoveTemp(Payloads->BreadcrumbAllocators));
-	#endif
-		delete Payloads;
-	}
-
-	SubmitPayloads(MoveTemp(AllPayloads));
-
-	#if WITH_RHI_BREADCRUMBS
-	// Enqueue the breadcrumb allocator references for cleanup once all prior payloads have completed on the GPU.
-	DeferredDelete([Array = MoveTemp(BreadcrumbAllocators)]() {});
-	#endif
-
-#endif
 }
 
 void FD3D12DynamicRHI::SubmitPayloads(TArray<FD3D12Payload*>&& Payloads)
@@ -481,44 +435,9 @@ FD3D12DynamicRHI::FProcessResult FD3D12DynamicRHI::ProcessSubmissionQueue()
 						FD3D12Queue& TargetQueue = CommandList->Device->GetQueue(CommandList->QueueType);
 						const uint32 MaxBatchSize = GetMaxExecuteBatchSize(TargetQueue.QueueType);
 
-						// Occlusion + Pipeline Stats Queries
-						TargetQueue.BatchedObjects.OcclusionQueries.Append(MoveTemp(CommandList->State.OcclusionQueries));
+						TargetQueue.BatchedObjects.OcclusionQueries    .Append(MoveTemp(CommandList->State.OcclusionQueries    ));
 						TargetQueue.BatchedObjects.PipelineStatsQueries.Append(MoveTemp(CommandList->State.PipelineStatsQueries));
-
-#if RHI_NEW_GPU_PROFILER
-						TargetQueue.BatchedObjects.TimestampQueries.Append(MoveTemp(CommandList->State.TimestampQueries));
-#else
-						// Timestamp Queries
-						if (CommandList->State.BeginTimestamp)
-						{
-							// Keep only the first Begin() in the batch
-							if (TargetQueue.NumCommandListsInBatch++ == 0)
-							{
-								TargetQueue.BatchedObjects.TimestampQueries.Emplace(MoveTemp(CommandList->State.BeginTimestamp));
-							}
-							else
-							{
-								// Remove the previous End() timestamp, to join the range together.
-								check(TargetQueue.BatchedObjects.TimestampQueries.Last().Type == ED3D12QueryType::CommandListEnd);
-								TargetQueue.BatchedObjects.TimestampQueries.RemoveAt(TargetQueue.BatchedObjects.TimestampQueries.Num() - 1);
-							}
-
-							TargetQueue.BatchedObjects.TimestampQueries.Append(MoveTemp(CommandList->State.TimestampQueries));
-							TargetQueue.BatchedObjects.TimestampQueries.Emplace(MoveTemp(CommandList->State.EndTimestamp));
-
-							if (TargetQueue.NumCommandListsInBatch >= MaxBatchSize)
-							{
-								// Start a new batch
-								TargetQueue.NumCommandListsInBatch = 0;
-							}
-						}
-						else
-						{
-							// No begin timestamp means timestamps aren't supported on this queue
-							check(CommandList->State.TimestampQueries.IsEmpty());
-							check(!CommandList->State.EndTimestamp);
-						}
-#endif
+						TargetQueue.BatchedObjects.TimestampQueries    .Append(MoveTemp(CommandList->State.TimestampQueries    ));
 					};
 
 					for (int32 Index = 0; Index < Payload->CommandListsToExecute.Num(); Index++)
@@ -555,8 +474,6 @@ uint64 FD3D12Queue::FinalizePayload(bool bRequiresSignal, FPayloadArray& Payload
 	check(PayloadToSubmit && this == &PayloadToSubmit->Queue);
 	check(PayloadToSubmit->SyncPointsToWait.Num() == 0);
 
-	NumCommandListsInBatch = 0;
-
 	BarrierTimestamps.CloseAndReset(PayloadToSubmit->BatchedObjects.QueryRanges);
 
 	// Gather query ranges from this payload, grouping by heap pointer
@@ -743,13 +660,11 @@ void FD3D12DynamicRHI::FlushBatchedPayloads(FD3D12Queue::FPayloadArray& Payloads
 		// Wait for queue fences
 		for (auto& [LocalFence, Value] : Payload->QueueFencesToWait)
 		{
-		#if RHI_NEW_GPU_PROFILER
 			Payload->EventStream.Emplace<UE::RHI::GPUProfiler::FEvent::FWaitFence>(
 				  FPlatformTime::Cycles64()
 				, Value
 				, LocalFence.OwnerQueue->GetProfilerQueue()
 			);
-		#endif
 
 			VERIFYD3D12RESULT(Queue.D3DCommandQueue->Wait(LocalFence.D3DFence, Value));
 		}
@@ -786,10 +701,7 @@ void FD3D12DynamicRHI::FlushBatchedPayloads(FD3D12Queue::FPayloadArray& Payloads
 			for (FD3D12CommandList* CommandList : Payload->CommandListsToExecute)
 			{
 				check(CommandList->IsClosed());
-
-#if RHI_NEW_GPU_PROFILER
 				CommandList->FlushProfilerEvents(Payload->EventStream, Time);
-#endif // RHI_NEW_GPU_PROFILER
 
 				D3DCommandLists.Add(CommandList->Interfaces.CommandList);
 
@@ -860,24 +772,20 @@ void FD3D12DynamicRHI::FlushBatchedPayloads(FD3D12Queue::FPayloadArray& Payloads
 		{
 			check(Queue.Fence.LastSignaledValue < Payload->CompletionFenceValue);
 
-		#if RHI_NEW_GPU_PROFILER
 			Payload->EventStream.Emplace<UE::RHI::GPUProfiler::FEvent::FSignalFence>(
 				  FPlatformTime::Cycles64()
 				, Payload->CompletionFenceValue
 			);
-		#endif
 
 			VERIFYD3D12RESULT(Queue.D3DCommandQueue->Signal(Queue.Fence.D3DFence, Payload->CompletionFenceValue));
 			Queue.Fence.LastSignaledValue.store(Payload->CompletionFenceValue, std::memory_order_release);
 		}
 
-	#if RHI_NEW_GPU_PROFILER
 		if (Payload->EndFrameEvent.IsSet())
 		{
 			Payload->EndFrameEvent->CPUTimestamp = FPlatformTime::Cycles64();
 			Payload->EventStream.Emplace<UE::RHI::GPUProfiler::FEvent::FFrameBoundary>(*Payload->EndFrameEvent);
 		}
-	#endif
 
 		if (Payload->PresentArgs.IsValid())
 		{
@@ -902,7 +810,6 @@ void FD3D12DynamicRHI::FlushBatchedPayloads(FD3D12Queue::FPayloadArray& Payloads
 
 		Payload->Queue.PendingInterrupt.Enqueue(Payload);
 
-#if RHI_NEW_GPU_PROFILER
 		if (Payload->Timing.IsSet())
 		{
 			Flush();
@@ -919,7 +826,6 @@ void FD3D12DynamicRHI::FlushBatchedPayloads(FD3D12Queue::FPayloadArray& Payloads
 				}
 			}
 		}
-#endif // RHI_NEW_GPU_PROFILER
 
 		if (Payload->HasWaitWork())
 		{
@@ -1264,7 +1170,6 @@ FD3D12DynamicRHI::FProcessResult FD3D12DynamicRHI::ProcessInterruptQueue()
 							}
 							break;
 
-					#if RHI_NEW_GPU_PROFILER
 						case ED3D12QueryType::ProfilerTimestampTOP:
 						case ED3D12QueryType::ProfilerTimestampBOP:
 							{
@@ -1277,27 +1182,16 @@ FD3D12DynamicRHI::FProcessResult FD3D12DynamicRHI::ProcessInterruptQueue()
 								Target = CPUDelta + CurrentQueue.Timing->CPUTimestamp;
 							}
 							break;
-					#else
-						case ED3D12QueryType::CommandListBegin:
-						case ED3D12QueryType::CommandListEnd:
-						case ED3D12QueryType::IdleBegin:
-						case ED3D12QueryType::IdleEnd:
-							check(CurrentQueue.Timing);
-							CurrentQueue.Timing->Timestamps.Add(Query.GetResult<uint64>());
-							break;
-					#endif
 						}
 					}
 				}
 			}
 
-		#if RHI_NEW_GPU_PROFILER
 			if (!Payload->EventStream.IsEmpty())
 			{
 				check(CurrentQueue.Timing);
 				CurrentQueue.Timing->EventStream.Append(MoveTemp(Payload->EventStream));
 			}
-		#endif
 
 			if (Payload->Timing.IsSet())
 			{
@@ -1336,9 +1230,7 @@ FD3D12DynamicRHI::FProcessResult FD3D12DynamicRHI::ProcessInterruptQueue()
 
 FD3D12PayloadBase::FD3D12PayloadBase(FD3D12Queue& Queue)
 	: Queue(Queue)
-#if RHI_NEW_GPU_PROFILER
 	, EventStream(Queue.GetProfilerQueue())
-#endif
 {}
 
 FD3D12PayloadBase::~FD3D12PayloadBase()
@@ -1357,19 +1249,8 @@ void FD3D12PayloadBase::PreExecute()
 	}
 }
 
-#ifndef D3D12_PREFER_QUERIES_FOR_GPU_TIME
-#define D3D12_PREFER_QUERIES_FOR_GPU_TIME 0
-#endif
-
-static TAutoConsoleVariable<int32> CVarGPUTimeFromTimestamps(
-	TEXT("r.D3D12.GPUTimeFromTimestamps"),
-	D3D12_PREFER_QUERIES_FOR_GPU_TIME,
-	TEXT("Prefer timestamps instead of GetHardwareGPUFrameTime to compute GPU frame time"),
-	ECVF_RenderThreadSafe);
-
 void FD3D12DynamicRHI::ProcessTimestamps(FD3D12TimingArray const& TimingPerQueue)
 {
-#if RHI_NEW_GPU_PROFILER
 	{
 		TArray<UE::RHI::GPUProfiler::FEventStream, TInlineAllocator<GD3D12MaxNumQueues>> Streams;
 		for (auto const& Timing : TimingPerQueue)
@@ -1378,58 +1259,6 @@ void FD3D12DynamicRHI::ProcessTimestamps(FD3D12TimingArray const& TimingPerQueue
 		}
 		UE::RHI::GPUProfiler::ProcessEvents(Streams);
 	}
-#else
-
-	// The total number of cycles where at least one GPU pipe was busy during the frame.
-	uint64 UnionBusyCycles = 0;
-	int32 BusyPipes = 0;
-
-	uint64 LastMinCycles = 0;
-	bool bFirst = true;
-
-	// Process the time ranges from each pipe.
-	while (true)
-	{
-		// Find the next minimum timestamp
-		FD3D12Timing* NextMin = nullptr;
-		for (auto const& Current : TimingPerQueue)
-		{
-			if (Current->HasMoreTimestamps() && (!NextMin || Current->GetCurrentTimestamp() < NextMin->GetCurrentTimestamp()))
-			{
-				NextMin = Current.Get();
-			}
-		}
-
-		if (!NextMin)
-			break; // No more timestamps to process
-
-		if (!bFirst)
-		{
-			if (BusyPipes > 0 && NextMin->GetCurrentTimestamp() > LastMinCycles)
-			{
-				// Accumulate the union busy time across all pipes
-				UnionBusyCycles += NextMin->GetCurrentTimestamp() - LastMinCycles;
-			}
-
-			if (!NextMin->IsStartingWork())
-			{
-				// Accumulate the busy time for this pipe specifically.
-				NextMin->BusyCycles += NextMin->GetCurrentTimestamp() - NextMin->GetPreviousTimestamp();
-			}
-		}
-
-		LastMinCycles = NextMin->GetCurrentTimestamp();
-
-		BusyPipes += NextMin->IsStartingWork() ? 1 : -1;
-		check(BusyPipes >= 0);
-
-		NextMin->AdvanceTimestamp();
-		bFirst = false;
-	}
-
-	check(BusyPipes == 0);
-	
-#endif
 
 	D3D12_QUERY_DATA_PIPELINE_STATISTICS PipelineStats{};
 	for (auto const& Current : TimingPerQueue)
@@ -1448,39 +1277,4 @@ void FD3D12DynamicRHI::ProcessTimestamps(FD3D12TimingArray const& TimingPerQueue
 	SET_DWORD_STAT(STAT_D3D12RHI_HSInvocations, PipelineStats.HSInvocations);
 	SET_DWORD_STAT(STAT_D3D12RHI_DSInvocations, PipelineStats.DSInvocations);
 	SET_DWORD_STAT(STAT_D3D12RHI_CSInvocations, PipelineStats.CSInvocations);
-
-#if RHI_NEW_GPU_PROFILER == 0
-
-	// @todo mgpu - how to handle multiple devices / queues with potentially different timestamp frequencies?
-	FD3D12Device* Device = GetAdapter().GetDevice(0);
-	double Frequency = Device->GetTimestampFrequency(ED3D12QueueType::Direct);
-
-	const double Scale64 = 1.0 / (Frequency * FPlatformTime::GetSecondsPerCycle64());
-
-	// Update the global GPU frame time stats
-	SET_CYCLE_COUNTER(STAT_RHI_GPUTotalTime, FPlatformMath::TruncToInt(double(UnionBusyCycles) * Scale64));
-
-	double HardwareGPUTime = 0.0;
-	if (GetHardwareGPUFrameTime(HardwareGPUTime) && CVarGPUTimeFromTimestamps.GetValueOnAnyThread() == 0)
-	{
-		SET_CYCLE_COUNTER(STAT_RHI_GPUTotalTimeHW, HardwareGPUTime);
-		GRHIGPUFrameTimeHistory.PushFrameCycles(1.0 / FPlatformTime::GetSecondsPerCycle64(), HardwareGPUTime);
-	}
-	else
-	{
-		SET_CYCLE_COUNTER(STAT_RHI_GPUTotalTimeHW, 0);
-		GRHIGPUFrameTimeHistory.PushFrameCycles(Frequency, UnionBusyCycles);
-	}
-
-	for (auto const& Current : TimingPerQueue)
-	{
-		switch (Current->Queue.QueueType)
-		{
-		case ED3D12QueueType::Direct: SET_CYCLE_COUNTER(STAT_RHI_GPUTotalTimeGraphics    , FPlatformMath::TruncToInt(double(Current->BusyCycles) * Scale64)); break;
-		case ED3D12QueueType::Async : SET_CYCLE_COUNTER(STAT_RHI_GPUTotalTimeAsyncCompute, FPlatformMath::TruncToInt(double(Current->BusyCycles) * Scale64)); break;
-		case ED3D12QueueType::Copy  : SET_CYCLE_COUNTER(STAT_RHI_GPUTotalTimeCopy        , FPlatformMath::TruncToInt(double(Current->BusyCycles) * Scale64)); break;
-		}
-	}
-
-#endif
 }
diff --git a/Engine/Source/Runtime/D3D12RHI/Private/D3D12Submission.h b/Engine/Source/Runtime/D3D12RHI/Private/D3D12Submission.h
--- a/Engine/Source/Runtime/D3D12RHI/Private/D3D12Submission.h
+++ b/Engine/Source/Runtime/D3D12RHI/Private/D3D12Submission.h
@@ -282,10 +282,8 @@ struct FD3D12PayloadBase : public FD3D12PayloadBaseFixLayout
 	TSharedPtr<FRHIBreadcrumbAllocatorArray> BreadcrumbAllocators {};
 #endif
 
-#if RHI_NEW_GPU_PROFILER
 	UE::RHI::GPUProfiler::FEventStream EventStream;
 	TOptional<UE::RHI::GPUProfiler::FEvent::FFrameBoundary> EndFrameEvent;
-#endif
 
 	virtual ~FD3D12PayloadBase();
 
@@ -313,11 +311,11 @@ struct FD3D12PayloadBase : public FD3D12PayloadBaseFixLayout
 
 	virtual bool HasSignalWork() const
 	{
-		return RequiresQueueFenceSignal() || ManualFencesToSignal.Num() > 0 || SubmissionEvent != nullptr || PresentArgs.IsValid()
-		#if RHI_NEW_GPU_PROFILER
-			|| EndFrameEvent.IsSet()
-		#endif
-		;
+		return RequiresQueueFenceSignal()
+			|| ManualFencesToSignal.Num() > 0
+			|| SubmissionEvent != nullptr
+			|| PresentArgs.IsValid()
+			|| EndFrameEvent.IsSet();
 	}
 
 protected:
diff --git a/Engine/Source/Runtime/D3D12RHI/Private/D3D12UAV.cpp b/Engine/Source/Runtime/D3D12RHI/Private/D3D12UAV.cpp
--- a/Engine/Source/Runtime/D3D12RHI/Private/D3D12UAV.cpp
+++ b/Engine/Source/Runtime/D3D12RHI/Private/D3D12UAV.cpp
@@ -427,13 +427,6 @@ void FD3D12CommandContext::ClearUAV(TRHICommandList_RecursiveHazardous<FD3D12Com
 				Context.GraphicsCommandList()->ClearUnorderedAccessViewUint(GPUHandle, OfflineHandle, UAV.GetResource()->GetResource(), *reinterpret_cast<const UINT(*)[4]>(ClearValues), 0, nullptr);
 				Context.UpdateResidency(UAV.GetResource());
 				Context.ConditionalSplitCommandList();
-
-			#if (RHI_NEW_GPU_PROFILER == 0)
-				if (Context.IsDefaultContext())
-				{
-					ParentDevice->RegisterGPUWork(1);
-				}
-			#endif
 			});
 		}
 		else
diff --git a/Engine/Source/Runtime/D3D12RHI/Private/Windows/WindowsD3D12Device.cpp b/Engine/Source/Runtime/D3D12RHI/Private/Windows/WindowsD3D12Device.cpp
--- a/Engine/Source/Runtime/D3D12RHI/Private/Windows/WindowsD3D12Device.cpp
+++ b/Engine/Source/Runtime/D3D12RHI/Private/Windows/WindowsD3D12Device.cpp
@@ -1948,13 +1948,6 @@ void FD3D12DynamicRHI::Init()
 
 	InitializeSubmissionPipe();
 
-#if (RHI_NEW_GPU_PROFILER == 0)
-	for (TSharedPtr<FD3D12Adapter>& Adapter : ChosenAdapters)
-	{
-		FD3D12BufferedGPUTiming::Initialize(Adapter.Get());
-	}
-#endif
-
 	FRenderResource::InitPreRHIResources();
 	GIsRHIInitialized = true;
 }
