Flush log in D3D12MessageCallBack so the message is available in VS debugger at UE_DEBUG_BREAK

diff --git a/Engine/Source/Runtime/D3D12RHI/Private/D3D12Adapter.cpp b/Engine/Source/Runtime/D3D12RHI/Private/D3D12Adapter.cpp
--- a/Engine/Source/Runtime/D3D12RHI/Private/D3D12Adapter.cpp
+++ b/Engine/Source/Runtime/D3D12RHI/Private/D3D12Adapter.cpp
@@ -369,6 +369,7 @@ void D3D12MessageCallBack(D3D12_MESSAGE_CATEGORY Category, D3D12_MESSAGE_SEVERIT
 	if (FPlatformMisc::IsDebuggerPresent() && bBreakOnLog)
 	{
 		// When we get here, then it means that BreakOnSeverity was set for this error message, so request the debug break here as well when the debugger is attached
+		GLog->Flush();
 		UE_DEBUG_BREAK();
 	}
 
