Log if GPU is tile-based for Windows D3D12

diff --git a/Engine/Source/Runtime/D3D12RHI/Private/D3D12Adapter.cpp b/Engine/Source/Runtime/D3D12RHI/Private/D3D12Adapter.cpp
--- a/Engine/Source/Runtime/D3D12RHI/Private/D3D12Adapter.cpp
+++ b/Engine/Source/Runtime/D3D12RHI/Private/D3D12Adapter.cpp
@@ -401,6 +401,7 @@ FD3D12AdapterDesc::FD3D12AdapterDesc(const DXGI_ADAPTER_DESC& InDesc, int32 InAd
 	, ResourceBindingTier(DeviceInfo.ResourceBindingTier)
 	, ResourceHeapTier(DeviceInfo.ResourceHeapTier)
 	, MaxRHIFeatureLevel(DeviceInfo.MaxRHIFeatureLevel)
+	, bTileBasedRenderer(DeviceInfo.bTileBasedRenderer)
 	, bUMA(DeviceInfo.bUMA)
 	, bCacheCoherentUMA(DeviceInfo.bCacheCoherentUMA)
 	, bSupportsGPUUploadHeaps(DeviceInfo.bSupportsGPUUploadHeaps)
diff --git a/Engine/Source/Runtime/D3D12RHI/Private/D3D12Adapter.h b/Engine/Source/Runtime/D3D12RHI/Private/D3D12Adapter.h
--- a/Engine/Source/Runtime/D3D12RHI/Private/D3D12Adapter.h
+++ b/Engine/Source/Runtime/D3D12RHI/Private/D3D12Adapter.h
@@ -52,9 +52,10 @@ struct FD3D12DeviceBasicInfo
 	uint32                      NumDeviceNodes;
 	bool                        bSupportsWaveOps;
 	bool                        bSupportsAtomic64;
-	bool						bUMA;
-	bool						bCacheCoherentUMA;
-	bool						bSupportsGPUUploadHeaps;
+	bool                        bTileBasedRenderer;
+	bool                        bUMA;
+	bool                        bCacheCoherentUMA;
+	bool                        bSupportsGPUUploadHeaps;
 	ERHIFeatureLevel::Type      MaxRHIFeatureLevel;
 };
 
@@ -90,6 +91,9 @@ struct FD3D12AdapterDesc
 	/** Number of device nodes (read: GPUs) */
 	uint32 NumDeviceNodes = 1;
 
+	/** Whether the GPU is tile-based. */
+	bool bTileBasedRenderer = false;
+
 	/** Whether the GPU has unified memory. This can be used to identify integrated GPU as well. */
 	bool bUMA = false;
 
diff --git a/Engine/Source/Runtime/D3D12RHI/Private/Windows/WindowsD3D12Device.cpp b/Engine/Source/Runtime/D3D12RHI/Private/Windows/WindowsD3D12Device.cpp
--- a/Engine/Source/Runtime/D3D12RHI/Private/Windows/WindowsD3D12Device.cpp
+++ b/Engine/Source/Runtime/D3D12RHI/Private/Windows/WindowsD3D12Device.cpp
@@ -477,11 +477,12 @@ inline bool GetSupportsAtomic64(IDXGIAdapter* Adapter, ID3D12Device* Device)
 	return (D3D12Caps9.AtomicInt64OnTypedResourceSupported || CheckDeviceForEmulatedAtomic64Support(Adapter, Device));
 }
 
-inline bool GetDeviceUMAInformation(ID3D12Device* device, bool& OutUMA, bool& OutCacheCoherentUMA)
+inline bool GetDeviceArchitectureInformation(ID3D12Device* device, bool& OutTileBasedRenderer, bool& OutUMA, bool& OutCacheCoherentUMA)
 {
 	D3D12_FEATURE_DATA_ARCHITECTURE data = { 0 };
 	if (SUCCEEDED(device->CheckFeatureSupport(D3D12_FEATURE_ARCHITECTURE, &data, sizeof(D3D12_FEATURE_DATA_ARCHITECTURE))))
 	{
+		OutTileBasedRenderer = (bool)(data.TileBasedRenderer);
 		OutUMA = (bool)(data.UMA);
 		OutCacheCoherentUMA = (bool)(data.CacheCoherentUMA);
 		return true;
@@ -523,9 +524,9 @@ static bool SafeTestD3D12CreateDevice(IDXGIAdapter* Adapter, D3D_FEATURE_LEVEL M
 
 			OutInfo.MaxRHIFeatureLevel = FindMaxRHIFeatureLevel(OutInfo.MaxFeatureLevel, OutInfo.MaxShaderModel, OutInfo.ResourceBindingTier, OutInfo.bSupportsWaveOps, OutInfo.bSupportsAtomic64);
 
-			if (GetDeviceUMAInformation(Device, OutInfo.bUMA, OutInfo.bCacheCoherentUMA) == false)
+			if (!GetDeviceArchitectureInformation(Device, OutInfo.bTileBasedRenderer, OutInfo.bUMA, OutInfo.bCacheCoherentUMA))
 			{
-				OutInfo.bUMA = OutInfo.bCacheCoherentUMA = false;
+				OutInfo.bTileBasedRenderer = OutInfo.bUMA = OutInfo.bCacheCoherentUMA = false;
 			}
 
 			OutInfo.bSupportsGPUUploadHeaps = GetSupportGPUUploadHeaps(Device);
@@ -945,11 +946,12 @@ void FD3D12DynamicRHIModule::FindAdapter()
 				);
 
 				UE_LOG(LogD3D12RHI, Log,
-					TEXT("  Adapter has %uMB of dedicated video memory, %uMB of dedicated system memory, and %uMB of shared system memory, %d output[s], UMA:%s CacheCoherentUMA:%s GPU Upload Heaps:%s"),
+					TEXT("  Adapter has %uMB of dedicated video memory, %uMB of dedicated system memory, and %uMB of shared system memory, %d output[s], TileBased:%s UMA:%s CacheCoherentUMA:%s GPU Upload Heaps:%s"),
 					(uint32)(AdapterDesc.DedicatedVideoMemory / (1024 * 1024)),
 					(uint32)(AdapterDesc.DedicatedSystemMemory / (1024 * 1024)),
 					(uint32)(AdapterDesc.SharedSystemMemory / (1024 * 1024)),
 					OutputCount,
+					DeviceInfo.bTileBasedRenderer ? TEXT("true") : TEXT("false"),
 					DeviceInfo.bUMA ? TEXT("true") : TEXT("false"),
 					DeviceInfo.bCacheCoherentUMA ? TEXT("true") : TEXT("false"),
 					DeviceInfo.bSupportsGPUUploadHeaps ? TEXT("true") : TEXT("false")
