Discard textures on creation for D3D12 to ensure we always have a valid resource

diff --git a/Engine/Source/Runtime/D3D12RHI/Private/D3D12RHIPrivate.h b/Engine/Source/Runtime/D3D12RHI/Private/D3D12RHIPrivate.h
--- a/Engine/Source/Runtime/D3D12RHI/Private/D3D12RHIPrivate.h
+++ b/Engine/Source/Runtime/D3D12RHI/Private/D3D12RHIPrivate.h
@@ -602,6 +602,7 @@ public:
 		FD3D12Texture* Texture{};
 		ED3D12Access CreateD3D12Access{};
 		ED3D12Access DesiredD3D12Access{};
+		bool bNeedsDiscard{};
 	};
 
 	// Per-platform D3D12Texture creation.
diff --git a/Engine/Source/Runtime/D3D12RHI/Private/D3D12Texture.cpp b/Engine/Source/Runtime/D3D12RHI/Private/D3D12Texture.cpp
--- a/Engine/Source/Runtime/D3D12RHI/Private/D3D12Texture.cpp
+++ b/Engine/Source/Runtime/D3D12RHI/Private/D3D12Texture.cpp
@@ -491,7 +491,6 @@ void SafeCreateTexture2D(
 	ED3D12Access InDefaultD3D12Access,
 	const TCHAR* Name)
 {
-
 #if GUARDED_TEXTURE_CREATES
 	bool bDriverCrash = true;
 	__try
@@ -797,10 +796,19 @@ FD3D12DynamicRHI::FCreateTextureInternalResult FD3D12DynamicRHI::CreateTextureIn
 
 	FD3D12TextureStats::D3D12TextureAllocated(*D3D12TextureOut);
 
+	// Discard for all RTV/DSV to be simple and safe.
+	// THe debug layer reports these require Discard/Clear/Copy as the first operation:
+	// - Placed resources
+	// - Reserved resources
+	// - Committed resources with D3D12_HEAP_FLAG_CREATE_NOT_ZEROED flag
+	const bool bNeedsDiscard = (DesiredD3D12Access != ED3D12Access::Discard) &&
+		EnumHasAnyFlags(ResourceDesc.Flags, D3D12_RESOURCE_FLAG_ALLOW_RENDER_TARGET | D3D12_RESOURCE_FLAG_ALLOW_DEPTH_STENCIL);
+
 	FCreateTextureInternalResult Result{};
 	Result.Texture = D3D12TextureOut;
 	Result.CreateD3D12Access = InitialD3D12Access;
 	Result.DesiredD3D12Access = DesiredD3D12Access;
+	Result.bNeedsDiscard = bNeedsDiscard;
 
 	return Result;
 #else
@@ -814,16 +822,71 @@ void UE::D3D12RHI::TextureUtils::ReconcileInitialState(
 	FRHICommandListBase& RHICmdList,
 	FD3D12Texture* Texture,
 	ED3D12Access CurrentD3D12Access,
-	ED3D12Access NeededD3D12Access)
+	ED3D12Access NeededD3D12Access,
+	bool bNeedsDiscard)
 {
-	if (CurrentD3D12Access != NeededD3D12Access)
+	if (bNeedsDiscard)
 	{
 		RHICmdList.EnqueueLambda([Texture, CurrentD3D12Access, NeededD3D12Access](FRHICommandListBase& ExecutingCmdList)
 		{
 			for (FD3D12Texture::FLinkedObjectIterator TextureIt(Texture); TextureIt; ++TextureIt)
 			{
+				const ETextureCreateFlags Flags = TextureIt->GetDesc().Flags;
 				const uint32 GPUIndex = TextureIt->GetParentDevice()->GetGPUIndex();
+				FD3D12CommandContext& Context = FD3D12CommandContext::Get(ExecutingCmdList, GPUIndex);
+				FD3D12Resource* const Resource = TextureIt->GetResource();
+
+				ED3D12Access NextD3D12Access = CurrentD3D12Access;
+				if (Context.GetPipeline() == ERHIPipeline::AsyncCompute)
+				{
+					if (!EnumHasAnyFlags(Flags, TexCreate_UAV))
+					{
+						checkf(false, TEXT("Can't discard non-UAVs for ERHIPipeline::AsyncCompute."));
+						goto SkipDiscard;
+					}
+
+					NextD3D12Access = ED3D12Access::UAVMask;
+				}
+				else if (EnumHasAnyFlags(Flags, TexCreate_RenderTargetable))
+				{
+					NextD3D12Access = ED3D12Access::RTV;
+				}
+				else if (EnumHasAnyFlags(Flags, TexCreate_DepthStencilTargetable))
+				{
+					NextD3D12Access = ED3D12Access::DSVWrite;
+				}
+
+				if (CurrentD3D12Access != NextD3D12Access)
+				{
+					Context.AddBarrier(
+						Resource,
+						CurrentD3D12Access,
+						NextD3D12Access,
+						D3D12_RESOURCE_BARRIER_ALL_SUBRESOURCES);
+					Context.FlushResourceBarriers();
+				}
 
+				Context.GetCommandList().GraphicsCommandList()->DiscardResource(Resource->GetResource(), nullptr);
+
+			SkipDiscard:
+				if (NextD3D12Access != NeededD3D12Access)
+				{
+					Context.AddBarrier(
+						Resource,
+						NextD3D12Access,
+						NeededD3D12Access,
+						D3D12_RESOURCE_BARRIER_ALL_SUBRESOURCES);
+				}
+			}
+		});
+	}
+	else if (CurrentD3D12Access != NeededD3D12Access)
+	{
+		RHICmdList.EnqueueLambda([Texture, CurrentD3D12Access, NeededD3D12Access](FRHICommandListBase& ExecutingCmdList)
+		{
+			for (FD3D12Texture::FLinkedObjectIterator TextureIt(Texture); TextureIt; ++TextureIt)
+			{
+				const uint32 GPUIndex = TextureIt->GetParentDevice()->GetGPUIndex();
 				FD3D12CommandContext& Context = FD3D12CommandContext::Get(ExecutingCmdList, GPUIndex);
 				Context.AddBarrier(
 					TextureIt->GetResource(),
@@ -976,7 +1039,7 @@ static UE::RHICore::FBaseTextureInitializerImplementation CreateTextureInitializ
 	return UE::RHICore::FBaseTextureInitializerImplementation(RHICmdList, Texture, WritableData, TextureSize,
 		[Texture = TRefCountPtr<FD3D12Texture>(Texture), CreateD3D12Access = CreateResult.CreateD3D12Access, DesiredD3D12Access = CreateResult.DesiredD3D12Access, UploadLocation = MoveTemp(UploadLocation)](FRHICommandListBase& RHICmdList) mutable
 		{
-			UE::D3D12RHI::TextureUtils::ReconcileInitialState(RHICmdList, Texture, CreateD3D12Access, ED3D12Access::CopyDest);
+			UE::D3D12RHI::TextureUtils::ReconcileInitialState(RHICmdList, Texture, CreateD3D12Access, ED3D12Access::CopyDest, false);
 
 			Texture->UploadInitialData(RHICmdList, MoveTemp(UploadLocation), DesiredD3D12Access);
 			return TRefCountPtr<FRHITexture>(MoveTemp(Texture));
@@ -1029,9 +1092,9 @@ FRHITextureInitializer FD3D12DynamicRHI::RHICreateTextureInitializer(FRHICommand
 	{
 		// Just return the texture with its default contents
 		return UE::RHICore::FBaseTextureInitializerImplementation(RHICmdList, CreateResult.Texture,
-			[Texture = TRefCountPtr<FD3D12Texture>(CreateResult.Texture), CreateD3D12Access = CreateResult.CreateD3D12Access, DesiredD3D12Access = CreateResult.DesiredD3D12Access](FRHICommandListBase& RHICmdList) mutable
+			[Texture = TRefCountPtr<FD3D12Texture>(CreateResult.Texture), CreateD3D12Access = CreateResult.CreateD3D12Access, DesiredD3D12Access = CreateResult.DesiredD3D12Access, bNeedsDiscard = CreateResult.bNeedsDiscard](FRHICommandListBase& RHICmdList) mutable
 			{
-				UE::D3D12RHI::TextureUtils::ReconcileInitialState(RHICmdList, Texture, CreateD3D12Access, DesiredD3D12Access);
+				UE::D3D12RHI::TextureUtils::ReconcileInitialState(RHICmdList, Texture, CreateD3D12Access, DesiredD3D12Access, bNeedsDiscard);
 				return TRefCountPtr<FRHITexture>(MoveTemp(Texture));
 			}
 		);
@@ -1569,7 +1632,7 @@ FTextureRHIRef FD3D12DynamicRHI::AsyncReallocateTexture2D_RenderThread(FRHIComma
 
 	// Allocate a new texture.
 	const FCreateTextureInternalResult CreateResult = CreateTextureInternal(CreateDesc, nullptr);
-	UE::D3D12RHI::TextureUtils::ReconcileInitialState(RHICmdList, CreateResult.Texture, CreateResult.CreateD3D12Access, CreateResult.DesiredD3D12Access);
+	UE::D3D12RHI::TextureUtils::ReconcileInitialState(RHICmdList, CreateResult.Texture, CreateResult.CreateD3D12Access, CreateResult.DesiredD3D12Access, CreateResult.bNeedsDiscard);
 
 	FD3D12Texture* NewTexture = CreateResult.Texture;
 
diff --git a/Engine/Source/Runtime/D3D12RHI/Private/D3D12Texture.h b/Engine/Source/Runtime/D3D12RHI/Private/D3D12Texture.h
--- a/Engine/Source/Runtime/D3D12RHI/Private/D3D12Texture.h
+++ b/Engine/Source/Runtime/D3D12RHI/Private/D3D12Texture.h
@@ -284,7 +284,7 @@ namespace UE::D3D12RHI::TextureUtils
 	// Fills in the ClearValue argument based on the ResourceDesc and CreateDesc flags.
 	D3D12_CLEAR_VALUE* FillClearValue(D3D12_CLEAR_VALUE& ClearValue, const FD3D12ResourceDesc& ResourceDesc, const FRHITextureDesc& TextureDesc);
 
-	void ReconcileInitialState(FRHICommandListBase& RHICmdList, FD3D12Texture* Texture, ED3D12Access CurrentD3D12Access, ED3D12Access NeededD3D12Access);
+	void ReconcileInitialState(FRHICommandListBase& RHICmdList, FD3D12Texture* Texture, ED3D12Access CurrentD3D12Access, ED3D12Access NeededD3D12Access, bool bNeedsDiscard);
 
 	uint64 CalculateResourceSize(FD3D12Texture* Texture);
 	void CopyBulkData(void* UploadMemory, uint64 UploadMemorySize, FD3D12Texture* Texture, FResourceBulkDataInterface* BulkData);
diff --git a/Engine/Source/Runtime/RenderCore/Private/RenderTargetPool.cpp b/Engine/Source/Runtime/RenderCore/Private/RenderTargetPool.cpp
--- a/Engine/Source/Runtime/RenderCore/Private/RenderTargetPool.cpp
+++ b/Engine/Source/Runtime/RenderCore/Private/RenderTargetPool.cpp
@@ -100,7 +100,8 @@ FPooledRenderTarget* FRenderTargetPool::CreateRenderTarget(FRHICommandListBase&
 #endif
 	}
 
-	const ERHIAccess AccessInitial = ERHIAccess::SRVMask;
+	const ETextureCreateFlags TexCreateFlags = Desc.Flags;
+	const ERHIAccess AccessInitial = RHIGetDefaultResourceState(TexCreateFlags, false);
 	FRHITextureCreateDesc CreateDesc(Desc, AccessInitial, Name);
 	const static FLazyName ClassName(TEXT("FPooledRenderTarget"));
 	CreateDesc.SetClassName(ClassName);
@@ -113,11 +114,8 @@ FPooledRenderTarget* FRenderTargetPool::CreateRenderTarget(FRHICommandListBase&
 	PooledRenderTargets.Add(Result);
 	PooledRenderTargetHashes.Add(DescHash);
 
-	if (EnumHasAnyFlags(Desc.Flags, TexCreate_UAV))
+	if (EnumHasAnyFlags(TexCreateFlags, TexCreate_UAV))
 	{
-		// The render target desc is invalid if a UAV is requested with an RHI that doesn't support the high-end feature level.
-		check(GMaxRHIFeatureLevel >= ERHIFeatureLevel::SM5 || GMaxRHIFeatureLevel == ERHIFeatureLevel::ES3_1);
-
 		if (GRHISupportsUAVFormatAliasing)
 		{
 			EPixelFormat AliasFormat = Desc.UAVFormat != PF_Unknown
