Cleaning up D3D12IntelExtensions.
- Reducing code copying.
- Removing static init dependency between a cvar and an array.

diff --git a/Engine/Source/Runtime/D3D12RHI/Private/D3D12IntelExtensions.cpp b/Engine/Source/Runtime/D3D12RHI/Private/D3D12IntelExtensions.cpp
--- a/Engine/Source/Runtime/D3D12RHI/Private/D3D12IntelExtensions.cpp
+++ b/Engine/Source/Runtime/D3D12RHI/Private/D3D12IntelExtensions.cpp
@@ -7,44 +7,48 @@
 #if INTEL_EXTENSIONS
 
 #if INTEL_GPU_CRASH_DUMPS
-
 namespace UE::RHICore::Intel::GPUCrashDumps
 {
 	extern RHICORE_API TAutoConsoleVariable<int32> CVarIntelCrashDumps;
 }
-
 #endif // INTEL_GPU_CRASH_DUMPS
 
-const INTCSupportedVersion SupportedExtensionsVersion[] = {
-	{ { 4, 8, 0 }, nullptr },														// Emulated Typed 64bit Atomics - this is required to run Nanite on ACM (DG2)
+bool MatchIntelExtensionVersion(const INTCExtensionVersion& ExtensionsVersion)
+{
+	// Make sure these are sorted with the highest version first (most features)
+	const INTCSupportedVersion SupportedExtensionsVersions[] =
+	{
 #if INTEL_GPU_CRASH_DUMPS
-	{ { 4, 14, 0 }, &( *UE::RHICore::Intel::GPUCrashDumps::CVarIntelCrashDumps) }	// Intel GPU crash dumps
+		// Intel GPU crash dumps
+		{
+			{ 4, 14, 0 },
+			&(*UE::RHICore::Intel::GPUCrashDumps::CVarIntelCrashDumps)
+		},
 #endif // INTEL_GPU_CRASH_DUMPS
-};
 
-bool MatchIntelExtensionVersion(const INTCExtensionVersion& ExtensionsVersion)
-{
-	uint32_t count = sizeof( SupportedExtensionsVersion ) / sizeof( SupportedExtensionsVersion[ 0 ] );
+		// Emulated Typed 64bit Atomics - this is required to run Nanite on ACM (DG2)
+		{
+			{ 4, 8, 0 },
+			nullptr
+		},
+	};
 
-	// We try to match the highest version first (most features)
-	for( int i = count - 1; i >= 0; i-- )
+	for (const INTCSupportedVersion& SupportedVersion : SupportedExtensionsVersions)
 	{
-		const INTCSupportedVersion& SupportedVersion = SupportedExtensionsVersion[ i ];
-
-		if( SupportedVersion.CVar && SupportedVersion.CVar->GetInt() == 0 )
+		if (SupportedVersion.CVar && SupportedVersion.CVar->GetInt() == 0)
 		{
 			continue;			// skip disabled feature
 		}
 
-		if( ( ExtensionsVersion.HWFeatureLevel == SupportedVersion.Version.HWFeatureLevel ) &&
-			( ExtensionsVersion.APIVersion == SupportedVersion.Version.APIVersion ) &&
-			( ExtensionsVersion.Revision == SupportedVersion.Version.Revision ) )
+		if (ExtensionsVersion.HWFeatureLevel == SupportedVersion.Version.HWFeatureLevel &&
+			ExtensionsVersion.APIVersion == SupportedVersion.Version.APIVersion &&
+			ExtensionsVersion.Revision == SupportedVersion.Version.Revision)
 		{
 			// First match wins
 			return true;
 		}
 	}
-	
+
 	// No matches
 	return false;
 }
@@ -52,38 +56,48 @@ bool MatchIntelExtensionVersion(const INTCExtensionVersion& ExtensionsVersion)
 #if INTEL_GPU_CRASH_DUMPS
 namespace UE::RHICore::Intel::GPUCrashDumps::D3D12
 {
-	uint64_t RegisterCommandList( ID3D12CommandList* pCommandList )
+	inline INTCExtensionContext* GetIntelExtensionContext()
+	{
+		return FD3D12DynamicRHI::GetD3DRHI()->GetIntelExtensionContext();
+	}
+
+	inline constexpr uint32 GetEventMarkerBreadcrumbType()
 	{
-		if( !IsEnabled() )
+#if INTEL_BREADCRUMBS_USE_BREADCRUMB_PTRS
+		return INTC_EVENT_MARKER_PTR;
+#else
+		return INTC_EVENT_MARKER_WSTRING;
+#endif
+	}
+
+	uint64_t RegisterCommandList(ID3D12CommandList* pCommandList)
+	{
+		if (!IsEnabled())
 		{
 			return 0;
 		}
 
-		return INTC_D3D12_GetCommandListHandle( FD3D12DynamicRHI::GetD3DRHI()->GetIntelExtensionContext(), pCommandList );
+		return INTC_D3D12_GetCommandListHandle(GetIntelExtensionContext(), pCommandList);
 	}
 
 #if WITH_RHI_BREADCRUMBS
 	void BeginBreadcrumb(ID3D12GraphicsCommandList* pCommandList, FRHIBreadcrumbNode* Breadcrumb)
 	{
-		if( !IsEnabled() )
+		if (!IsEnabled())
 		{
 			return;
 		}
-		
+
 		FMarker Marker(Breadcrumb);
 		if (Marker)
 		{
-#if INTEL_BREADCRUMBS_USE_BREADCRUMB_PTRS
-			INTC_D3D12_SetEventMarker( FD3D12DynamicRHI::GetD3DRHI()->GetIntelExtensionContext(), pCommandList, INTC_EVENT_MARKER_BEGIN | INTC_EVENT_MARKER_PTR, Marker.GetPtr(), Marker.GetSize() );
-#else
-			INTC_D3D12_SetEventMarker( FD3D12DynamicRHI::GetD3DRHI()->GetIntelExtensionContext(), pCommandList, INTC_EVENT_MARKER_BEGIN | INTC_EVENT_MARKER_WSTRING, Marker.GetPtr(), Marker.GetSize() );
-#endif
+			INTC_D3D12_SetEventMarker(GetIntelExtensionContext(), pCommandList, INTC_EVENT_MARKER_BEGIN | GetEventMarkerBreadcrumbType(), Marker.GetPtr(), Marker.GetSize());
 		}
 	}
 
 	void EndBreadcrumb(ID3D12GraphicsCommandList* pCommandList, FRHIBreadcrumbNode* Breadcrumb)
 	{
-		if( !IsEnabled() )
+		if (!IsEnabled())
 		{
 			return;
 		}
@@ -91,11 +105,7 @@ namespace UE::RHICore::Intel::GPUCrashDumps::D3D12
 		FMarker Marker(Breadcrumb);
 		if (Marker)
 		{
-#if INTEL_BREADCRUMBS_USE_BREADCRUMB_PTRS
-			INTC_D3D12_SetEventMarker( FD3D12DynamicRHI::GetD3DRHI()->GetIntelExtensionContext(), pCommandList, INTC_EVENT_MARKER_END | INTC_EVENT_MARKER_PTR, Marker.GetPtr(), Marker.GetSize() );
-#else
-			INTC_D3D12_SetEventMarker( FD3D12DynamicRHI::GetD3DRHI()->GetIntelExtensionContext(), pCommandList, INTC_EVENT_MARKER_END | INTC_EVENT_MARKER_WSTRING, Marker.GetPtr(), Marker.GetSize() );
-#endif
+			INTC_D3D12_SetEventMarker(GetIntelExtensionContext(), pCommandList, INTC_EVENT_MARKER_END | GetEventMarkerBreadcrumbType(), Marker.GetPtr(), Marker.GetSize());
 		}
 	}
 #endif
