Move debug information after the syncronization in order to avoid data race between GetBackBuffer_RHIThread and PresentOnSubmissionThread/SetBackBufferIndex_RHIThread

#Synced-CL 48277489

diff --git a/Engine/Source/Runtime/D3D12RHI/Private/D3D12Viewport.cpp b/Engine/Source/Runtime/D3D12RHI/Private/D3D12Viewport.cpp
--- a/Engine/Source/Runtime/D3D12RHI/Private/D3D12Viewport.cpp
+++ b/Engine/Source/Runtime/D3D12RHI/Private/D3D12Viewport.cpp
@@ -570,17 +570,17 @@ FD3D12Texture* FD3D12Viewport::GetBackBuffer_RHIThread(bool bWaitForLastPresent)
 {
 	TRACE_CPUPROFILER_EVENT_SCOPE(FD3D12Viewport::GetBackBuffer_RHIThread);
 
-#if CPUPROFILERTRACE_ENABLED
-	UE_TRACE_LOG_SCOPED_T(Cpu, GetBackBufferOnRhiThread, CpuChannel)
-		<< GetBackBufferOnRhiThread.BackBuffer(CurrentBackBufferIndex_RHIThread);
-#endif // CPUPROFILERTRACE_ENABLED
-
-	checkSlow(CurrentBackBuffer_RHIThread);
-
 	if (bWaitForLastPresent)
 	{
 		WaitForLastPresent();
 	}
+	
+	{
+#if CPUPROFILERTRACE_ENABLED
+		UE_TRACE_LOG_SCOPED_T(Cpu, GetBackBufferOnRhiThread, CpuChannel)
+		<< GetBackBufferOnRhiThread.BackBuffer(CurrentBackBufferIndex_RHIThread);
+#endif // CPUPROFILERTRACE_ENABLED
+	}
 
 	return CurrentBackBuffer_RHIThread->Texture;
 }
diff --git a/Engine/Source/Runtime/D3D12RHI/Private/D3D12Viewport.h b/Engine/Source/Runtime/D3D12RHI/Private/D3D12Viewport.h
--- a/Engine/Source/Runtime/D3D12RHI/Private/D3D12Viewport.h
+++ b/Engine/Source/Runtime/D3D12RHI/Private/D3D12Viewport.h
@@ -97,8 +97,6 @@ public:
 
 	FD3D12Texture* GetSDRBackBuffer_RHIThread(bool bWaitForLastPresent = true) const
 	{
-		checkSlow(CurrentBackBuffer_RHIThread);
-
 		if (bWaitForLastPresent)
 		{
 			WaitForLastPresent();
