Bindless UniformBuffer support.
- Adding GetCBVFromHeap to shaders for future use.
- Adding bindless handles to FRHIUniformBuffer.
- Adding support for adding a bindles FRHIUniformBuffer to RHIShaderParameters. This ends up needing all resources to be 'marked used' at binding time.

#Synced-CL 48355415

diff --git a/Engine/Shaders/Public/Platform/D3D/D3DCommon.ush b/Engine/Shaders/Public/Platform/D3D/D3DCommon.ush
--- a/Engine/Shaders/Public/Platform/D3D/D3DCommon.ush
+++ b/Engine/Shaders/Public/Platform/D3D/D3DCommon.ush
@@ -158,5 +158,6 @@ void UEReportAssertWithPayload(uint AssertID, uint4 Payload)
 
 	#define GetSRVFromHeap(Type, Index) ResourceDescriptorHeap[UE_BINDLESS_INDEX(Index)]
 	#define GetUAVFromHeap(Type, Index) ResourceDescriptorHeap[UE_BINDLESS_INDEX(Index)]
+	#define GetCBVFromHeap(Type, Index) ResourceDescriptorHeap[UE_BINDLESS_INDEX(Index)]
 	#define GetSamplerFromHeap(Type, Index) SamplerDescriptorHeap[UE_BINDLESS_INDEX(Index)]
 #endif	
diff --git a/Engine/Shaders/Public/Platform/Metal/MetalCommon.ush b/Engine/Shaders/Public/Platform/Metal/MetalCommon.ush
--- a/Engine/Shaders/Public/Platform/Metal/MetalCommon.ush
+++ b/Engine/Shaders/Public/Platform/Metal/MetalCommon.ush
@@ -84,6 +84,7 @@
 #if PLATFORM_SUPPORTS_BINDLESS
 	#define GetSRVFromHeap(Type, Index) ResourceDescriptorHeap[Index]
 	#define GetUAVFromHeap(Type, Index) ResourceDescriptorHeap[Index]
+	#define GetCBVFromHeap(Type, Index) ResourceDescriptorHeap[Index]
 	#define GetSamplerFromHeap(Type, Index) SamplerDescriptorHeap[Index]
 #endif
 
diff --git a/Engine/Shaders/Public/Platform/Vulkan/VulkanCommon.ush b/Engine/Shaders/Public/Platform/Vulkan/VulkanCommon.ush
--- a/Engine/Shaders/Public/Platform/Vulkan/VulkanCommon.ush
+++ b/Engine/Shaders/Public/Platform/Vulkan/VulkanCommon.ush
@@ -52,6 +52,7 @@ void ImageInterlockedOrUInt64(RWTexture2D<UlongType> Dest, uint2 Coord, UlongTyp
 	#define VULKAN_SAMPLER_HEAP(Type) VULKAN_HEAP_NAME(BINDLESS_SAMPLER_ARRAY_PREFIX,Type)
 	#define VULKAN_SRV_HEAP(Type) VULKAN_HEAP_NAME(BINDLESS_SRV_ARRAY_PREFIX,Type)
 	#define VULKAN_UAV_HEAP(Type) VULKAN_HEAP_NAME(BINDLESS_UAV_ARRAY_PREFIX,Type)
+	#define VULKAN_CBV_HEAP(Type) VULKAN_HEAP_NAME(BINDLESS_CBV_ARRAY_PREFIX,Type)
 
 	#if (RAYCALLABLESHADER || RAYHITGROUPSHADER || RAYGENSHADER || RAYMISSSHADER)
 		#define VULKAN_HEAP_ACCESS(BindlessIndex) NonUniformResourceIndex(BindlessIndex)
@@ -61,10 +62,12 @@ void ImageInterlockedOrUInt64(RWTexture2D<UlongType> Dest, uint2 Coord, UlongTyp
 
 	#define UE_BINDLESS_SRV_HEAP_DECL(Type) Type VULKAN_SRV_HEAP(Type)[];
 	#define UE_BINDLESS_UAV_HEAP_DECL(Type) Type VULKAN_UAV_HEAP(Type)[];
+	#define UE_BINDLESS_CBV_HEAP_DECL(Type) Type VULKAN_CBV_HEAP(Type)[];
 	#define UE_BINDLESS_SAMPLER_HEAP_DECL(Type) Type VULKAN_SAMPLER_HEAP(Type)[];
 
 	#define GetSRVFromHeap(Type, Index) VULKAN_SRV_HEAP(Type)[VULKAN_HEAP_ACCESS(Index)]
 	#define GetUAVFromHeap(Type, Index) VULKAN_UAV_HEAP(Type)[VULKAN_HEAP_ACCESS(Index)]
+	#define GetCBVFromHeap(Type, Index) VULKAN_CBV_HEAP(Type)[VULKAN_HEAP_ACCESS(Index)]
 	#define GetSamplerFromHeap(Type, Index) VULKAN_SAMPLER_HEAP(Type)[VULKAN_HEAP_ACCESS(Index)]
 #endif
 
diff --git a/Engine/Source/Developer/VulkanShaderFormat/Private/SpirVShaderCompiler.inl b/Engine/Source/Developer/VulkanShaderFormat/Private/SpirVShaderCompiler.inl
--- a/Engine/Source/Developer/VulkanShaderFormat/Private/SpirVShaderCompiler.inl
+++ b/Engine/Source/Developer/VulkanShaderFormat/Private/SpirVShaderCompiler.inl
@@ -1467,6 +1467,7 @@ static void ModifyCompilerInput(FSpirvShaderCompilerInternalState& InternalState
 
 	Input.Environment.SetDefine(TEXT("BINDLESS_SRV_ARRAY_PREFIX"), FShaderParameterParser::kBindlessSRVArrayPrefix);
 	Input.Environment.SetDefine(TEXT("BINDLESS_UAV_ARRAY_PREFIX"), FShaderParameterParser::kBindlessUAVArrayPrefix);
+	Input.Environment.SetDefine(TEXT("BINDLESS_CBV_ARRAY_PREFIX"), FShaderParameterParser::kBindlessCBVArrayPrefix);
 	Input.Environment.SetDefine(TEXT("BINDLESS_SAMPLER_ARRAY_PREFIX"), FShaderParameterParser::kBindlessSamplerArrayPrefix);
 
 	if (InternalState.IsAndroid())
diff --git a/Engine/Source/Runtime/Apple/MetalRHI/Private/MetalBindlessDescriptors.cpp b/Engine/Source/Runtime/Apple/MetalRHI/Private/MetalBindlessDescriptors.cpp
--- a/Engine/Source/Runtime/Apple/MetalRHI/Private/MetalBindlessDescriptors.cpp
+++ b/Engine/Source/Runtime/Apple/MetalRHI/Private/MetalBindlessDescriptors.cpp
@@ -251,12 +251,11 @@ void FMetalBindlessDescriptorManager::FreeDescriptor(FRHIDescriptorHandle Descri
 	}
 }
 
-void FMetalBindlessDescriptorManager::UpdateDescriptor(FRHICommandListBase& RHICmdList, FRHIDescriptorHandle DescriptorHandle, MTL::Texture* Texture, EMetalDescriptorUpdateType UpdateType)
+
+void FMetalBindlessDescriptorManager::UpdateDescriptor(FRHICommandListBase& RHICmdList, FRHIDescriptorHandle DescriptorHandle, IRDescriptorTableEntry Descriptor, EMetalDescriptorUpdateType InUpdateType)
 {
-	IRDescriptorTableEntry Descriptor = CreateDescriptor(Texture);
-	
-	UpdateType = !GIsRHIInitialized ? EMetalDescriptorUpdateType::Immediate : UpdateType;
-	
+	const EMetalDescriptorUpdateType UpdateType = !GIsRHIInitialized ? EMetalDescriptorUpdateType::Immediate : InUpdateType;
+
 	RHICmdList.EnqueueLambda([this, UpdateType, Descriptor, DescriptorHandle](FRHICommandListBase& RHICmdList)
 	{
 		if (UpdateType == EMetalDescriptorUpdateType::Immediate)
@@ -264,15 +263,27 @@ void FMetalBindlessDescriptorManager::UpdateDescriptor(FRHICommandListBase& RHIC
 			StandardResources.UpdateDescriptorImmediately(DescriptorHandle, Descriptor);
 		}
 		else
-		{	
+		{
 			FMetalRHICommandContext& Context = FMetalRHICommandContext::Get(RHICmdList);
 			Context.EnqueueDescriptorUpdate(DescriptorHandle, Descriptor);
 		}
 	});
-	
+
 	RHICmdList.RHIThreadFence(true);
 }
 
+void FMetalBindlessDescriptorManager::UpdateDescriptor(FRHICommandListBase& RHICmdList, FRHIDescriptorHandle DescriptorHandle, MTL::Texture* Texture, EMetalDescriptorUpdateType UpdateType)
+{
+	IRDescriptorTableEntry Descriptor = CreateDescriptor(Texture);
+	UpdateDescriptor(RHICmdList, DescriptorHandle, Descriptor, UpdateType);
+}
+
+void FMetalBindlessDescriptorManager::UpdateDescriptor(FRHICommandListBase& RHICmdList, FRHIDescriptorHandle DescriptorHandle, FMetalBuffer* Buffer, EMetalDescriptorUpdateType UpdateType)
+{
+	IRDescriptorTableEntry Descriptor = CreateDescriptor(Buffer);
+	UpdateDescriptor(RHICmdList, DescriptorHandle, Descriptor, UpdateType);
+}
+
 void FMetalBindlessDescriptorManager::UpdateDescriptor(FMetalRHICommandContext* Context, FRHIDescriptorHandle DescriptorHandle, FMetalResourceViewBase* Resource, EMetalDescriptorUpdateType UpdateType)
 {
 	IRDescriptorTableEntry Descriptor = CreateDescriptor(Resource);
diff --git a/Engine/Source/Runtime/Apple/MetalRHI/Private/MetalBindlessDescriptors.h b/Engine/Source/Runtime/Apple/MetalRHI/Private/MetalBindlessDescriptors.h
--- a/Engine/Source/Runtime/Apple/MetalRHI/Private/MetalBindlessDescriptors.h
+++ b/Engine/Source/Runtime/Apple/MetalRHI/Private/MetalBindlessDescriptors.h
@@ -106,7 +106,9 @@ public:
 
 	void FreeDescriptor(FRHIDescriptorHandle DescriptorHandle);
 
+	void UpdateDescriptor(FRHICommandListBase& RHICmdList, FRHIDescriptorHandle DescriptorHandle, IRDescriptorTableEntry Descriptor, EMetalDescriptorUpdateType UpdateType);
 	void UpdateDescriptor(FRHICommandListBase& RHICmdList, FRHIDescriptorHandle DescriptorHandle, MTL::Texture* Texture, EMetalDescriptorUpdateType UpdateType);
+	void UpdateDescriptor(FRHICommandListBase& RHICmdList, FRHIDescriptorHandle DescriptorHandle, FMetalBuffer* Buffer, EMetalDescriptorUpdateType UpdateType);
 	void UpdateDescriptor(FMetalRHICommandContext* Context, FRHIDescriptorHandle DescriptorHandle, FMetalResourceViewBase* Resource, EMetalDescriptorUpdateType UpdateType);
 	
 	void FlushPendingDescriptorUpdates(FMetalRHICommandContext& Context, const FMetalPendingDescriptorUpdates& PendingDescriptorUpdates);
diff --git a/Engine/Source/Runtime/Apple/MetalRHI/Private/MetalUniformBuffer.cpp b/Engine/Source/Runtime/Apple/MetalRHI/Private/MetalUniformBuffer.cpp
--- a/Engine/Source/Runtime/Apple/MetalRHI/Private/MetalUniformBuffer.cpp
+++ b/Engine/Source/Runtime/Apple/MetalRHI/Private/MetalUniformBuffer.cpp
@@ -15,6 +15,18 @@
 
 #pragma mark Suballocated Uniform Buffer Implementation
 
+#if PLATFORM_SUPPORTS_BINDLESS_RENDERING
+inline bool DoesUniformBufferNeedBindlessHandle(FMetalDevice& InDevice, const FRHIUniformBufferLayout* Layout)
+{
+	if (EnumHasAnyFlags(Layout->Flags, ERHIUniformBufferFlags::BindlessAccessible))
+	{
+		return IsMetalBindlessEnabled();
+	}
+
+	return false;
+}
+#endif
+
 FMetalSuballocatedUniformBuffer::FMetalSuballocatedUniformBuffer(FMetalDevice& InDevice, const void *Contents, const FRHIUniformBufferLayout* Layout,
 																EUniformBufferUsage Usage, EUniformBufferValidation InValidation)
     : FRHIUniformBuffer(Layout)
@@ -31,6 +43,20 @@ FMetalSuballocatedUniformBuffer::FMetalSuballocatedUniformBuffer(FMetalDevice& I
 		CopyResourceTable(Contents, ResourceTable);
 		PushToGPUBacking(Shadow);
 	}
+
+#if PLATFORM_SUPPORTS_BINDLESS_RENDERING
+	if (DoesUniformBufferNeedBindlessHandle(InDevice, Layout))
+	{
+		if (BackingBuffer)
+		{
+			BindlessHandle = Device.GetBindlessDescriptorManager()->AllocateDescriptor(ERHIDescriptorType::CBV, BackingBuffer.Get(), 0);
+		}
+		else
+		{
+			BindlessHandle = Device.GetBindlessDescriptorManager()->AllocateDescriptor(ERHIDescriptorType::CBV);
+		}
+	}
+#endif
 }
 
 FMetalSuballocatedUniformBuffer::~FMetalSuballocatedUniformBuffer()
@@ -40,15 +66,29 @@ FMetalSuballocatedUniformBuffer::~FMetalSuballocatedUniformBuffer()
 		FMetalDynamicRHI::Get().DeferredDelete(BackingBuffer);
 		BackingBuffer.Reset();
 	}
+
+#if PLATFORM_SUPPORTS_BINDLESS_RENDERING
+	if (BindlessHandle.IsValid())
+	{
+		FMetalDynamicRHI::Get().DeferredDelete(BindlessHandle);
+	}
+#endif
 	
 	FMemory::Free(Shadow);
 }
 
-void FMetalSuballocatedUniformBuffer::Update(const void* Contents)
+void FMetalSuballocatedUniformBuffer::Update(FRHICommandListBase& RHICmdList, const void* Contents)
 {
     UE::RHICore::UpdateUniformBufferConstants(Shadow, Contents, GetLayout());
 	CopyResourceTable(Contents, ResourceTable);
 	PushToGPUBacking(Shadow);
+
+#if PLATFORM_SUPPORTS_BINDLESS_RENDERING
+	if (BackingBuffer && BindlessHandle.IsValid())
+	{
+		Device.GetBindlessDescriptorManager()->UpdateDescriptor(RHICmdList, BindlessHandle, BackingBuffer.Get(), EMetalDescriptorUpdateType::GPU);
+	}
+#endif
 }
 
 // Acquires a region in the current frame's uniform buffer and
diff --git a/Engine/Source/Runtime/Apple/MetalRHI/Private/MetalUniformBufferCommands.cpp b/Engine/Source/Runtime/Apple/MetalRHI/Private/MetalUniformBufferCommands.cpp
--- a/Engine/Source/Runtime/Apple/MetalRHI/Private/MetalUniformBufferCommands.cpp
+++ b/Engine/Source/Runtime/Apple/MetalRHI/Private/MetalUniformBufferCommands.cpp
@@ -29,7 +29,7 @@ void FMetalDynamicRHI::RHIUpdateUniformBuffer(FRHICommandListBase& RHICmdList, F
 
 	RHICmdList.EnqueueLambda([UniformBuffer, SrcContents](FRHICommandListBase& RHICmdList)
 	{
-		UniformBuffer->Update(SrcContents);
+		UniformBuffer->Update(RHICmdList, SrcContents);
 	});
 	
 	RHICmdList.RHIThreadFence(true);
diff --git a/Engine/Source/Runtime/Apple/MetalRHI/Public/MetalUniformBuffer.h b/Engine/Source/Runtime/Apple/MetalRHI/Public/MetalUniformBuffer.h
--- a/Engine/Source/Runtime/Apple/MetalRHI/Public/MetalUniformBuffer.h
+++ b/Engine/Source/Runtime/Apple/MetalRHI/Public/MetalUniformBuffer.h
@@ -27,6 +27,10 @@ public:
 	// CPU side shadow memory to hold updates for single-draw or multi-frame buffers.
     // This allows you to upload on a frame but actually use this UB later on
     void* Shadow;
+
+#if PLATFORM_SUPPORTS_BINDLESS_RENDERING
+	FRHIDescriptorHandle BindlessHandle;
+#endif
 	
 private:
 	FMetalDevice& Device;
@@ -41,7 +45,7 @@ public:
 									EUniformBufferUsage Usage, EUniformBufferValidation Validation);
     ~FMetalSuballocatedUniformBuffer();
 
-	void Update(const void* Contents);
+	void Update(FRHICommandListBase& RHICmdList, const void* Contents);
 
 private:
 	// Copies the RDG resources to a resource table for a deferred update on the RHI thread.
diff --git a/Engine/Source/Runtime/D3D12RHI/Private/D3D12Commands.cpp b/Engine/Source/Runtime/D3D12RHI/Private/D3D12Commands.cpp
--- a/Engine/Source/Runtime/D3D12RHI/Private/D3D12Commands.cpp
+++ b/Engine/Source/Runtime/D3D12RHI/Private/D3D12Commands.cpp
@@ -620,34 +620,41 @@ static void SetShaderParametersOnContext(
 	{
 		if (FRHIResource* Resource = Parameter.Resource)
 		{
-			FRHIDescriptorHandle Handle;
-
-			switch (Parameter.Type)
+			if (Parameter.Type == FRHIShaderParameterResource::EType::UniformBuffer)
 			{
-			case FRHIShaderParameterResource::EType::Texture:
-				Handle = static_cast<FRHITexture*>(Resource)->GetDefaultBindlessHandle();
-				Binder.SetTexture(static_cast<FRHITexture*>(Resource), Parameter.Index);
-				break;
-			case FRHIShaderParameterResource::EType::ResourceView:
-				Handle = static_cast<FRHIShaderResourceView*>(Resource)->GetBindlessHandle();
-				Binder.SetSRV(static_cast<FRHIShaderResourceView*>(Resource), Parameter.Index);
-				break;
-			case FRHIShaderParameterResource::EType::UnorderedAccessView:
-				Handle = static_cast<FRHIUnorderedAccessView*>(Resource)->GetBindlessHandle();
-				Binder.SetUAV(static_cast<FRHIUnorderedAccessView*>(Resource), Parameter.Index, true);
-				break;
-			case FRHIShaderParameterResource::EType::Sampler:
-				Handle = static_cast<FRHISamplerState*>(Resource)->GetBindlessHandle();
-				Binder.SetSampler(static_cast<FRHISamplerState*>(Resource), Parameter.Index);
-				break;
-			case FRHIShaderParameterResource::EType::ResourceCollection:
-				Handle = static_cast<FRHIResourceCollection*>(Resource)->GetBindlessHandle();
-				Binder.SetResourceCollection(static_cast<FRHIResourceCollection*>(Resource), Parameter.Index);
-				break;
+				UE::RHI::Private::SetAllUniformBufferResourcesForBindless(Binder, static_cast<FRHIUniformBuffer*>(Resource));
 			}
+			else
+			{
+				FRHIDescriptorHandle Handle;
 
-			checkf(Handle.IsValid(), TEXT("D3D12 resource did not provide a valid descriptor handle. Please validate that all D3D12 types can provide this or that the resource is still valid."));
-			Binder.SetBindlessHandle(Handle, Parameter.Index);
+				switch (Parameter.Type)
+				{
+				case FRHIShaderParameterResource::EType::Texture:
+					Handle = static_cast<FRHITexture*>(Resource)->GetDefaultBindlessHandle();
+					Binder.SetTexture(static_cast<FRHITexture*>(Resource), Parameter.Index);
+					break;
+				case FRHIShaderParameterResource::EType::ResourceView:
+					Handle = static_cast<FRHIShaderResourceView*>(Resource)->GetBindlessHandle();
+					Binder.SetSRV(static_cast<FRHIShaderResourceView*>(Resource), Parameter.Index);
+					break;
+				case FRHIShaderParameterResource::EType::UnorderedAccessView:
+					Handle = static_cast<FRHIUnorderedAccessView*>(Resource)->GetBindlessHandle();
+					Binder.SetUAV(static_cast<FRHIUnorderedAccessView*>(Resource), Parameter.Index, true);
+					break;
+				case FRHIShaderParameterResource::EType::Sampler:
+					Handle = static_cast<FRHISamplerState*>(Resource)->GetBindlessHandle();
+					Binder.SetSampler(static_cast<FRHISamplerState*>(Resource), Parameter.Index);
+					break;
+				case FRHIShaderParameterResource::EType::ResourceCollection:
+					Handle = static_cast<FRHIResourceCollection*>(Resource)->GetBindlessHandle();
+					Binder.SetResourceCollection(static_cast<FRHIResourceCollection*>(Resource), Parameter.Index);
+					break;
+				}
+
+				checkf(Handle.IsValid(), TEXT("D3D12 resource did not provide a valid descriptor handle. Please validate that all D3D12 types can provide this or that the resource is still valid."));
+				Binder.SetBindlessHandle(Handle, Parameter.Index);
+			}
 		}
 	}
 #endif
diff --git a/Engine/Source/Runtime/D3D12RHI/Private/D3D12RayTracing.cpp b/Engine/Source/Runtime/D3D12RHI/Private/D3D12RayTracing.cpp
--- a/Engine/Source/Runtime/D3D12RHI/Private/D3D12RayTracing.cpp
+++ b/Engine/Source/Runtime/D3D12RHI/Private/D3D12RayTracing.cpp
@@ -1896,7 +1896,7 @@ public:
 			ShaderTable.RemoveUpdateUniformBufferListener(RecordIndex, this);
 		}
 
-		virtual void UniformBufferUpdated(FRHICommandListBase& CmdList, FD3D12UniformBuffer* InUpdatedUniformBuffer) final override
+		virtual void UniformBufferUpdated(const FD3D12ContextArray& Contexts, FD3D12UniformBuffer* InUpdatedUniformBuffer) final override
 		{
 			check(UniformBuffer == InUpdatedUniformBuffer);
 			ShaderTable.SetLocalShaderParameters(ShaderTableOffset, RecordIndex, OffsetWithinRootSignature, InUpdatedUniformBuffer->ResourceLocation.GetGPUVirtualAddress());
diff --git a/Engine/Source/Runtime/D3D12RHI/Private/D3D12Resources.h b/Engine/Source/Runtime/D3D12RHI/Private/D3D12Resources.h
--- a/Engine/Source/Runtime/D3D12RHI/Private/D3D12Resources.h
+++ b/Engine/Source/Runtime/D3D12RHI/Private/D3D12Resources.h
@@ -921,7 +921,7 @@ class FD3D12UniformBuffer;
 struct ID3D12UniformBufferUpdateListener
 {
 	virtual void RemoveListener(FD3D12UniformBuffer* InUniformBufferToRemove) = 0;
-	virtual void UniformBufferUpdated(FRHICommandListBase& CmdList, FD3D12UniformBuffer* InUpdatedUniformBuffer) = 0;
+	virtual void UniformBufferUpdated(const FD3D12ContextArray& Contexts, FD3D12UniformBuffer* InUpdatedUniformBuffer) = 0;
 };
 
 /** Uniform buffer resource class. */
@@ -931,9 +931,7 @@ private:
 	mutable FCriticalSection UpdateListenersCS;
 	TArray<ID3D12UniformBufferUpdateListener*> UpdateListeners;
 public:
-#if D3D12RHI_USE_CONSTANT_BUFFER_VIEWS
 	class FD3D12ConstantBufferView* View = nullptr;
-#endif
 
 	/** The D3D12 constant buffer resource */
 	FD3D12ResourceLocation ResourceLocation;
@@ -951,6 +949,8 @@ public:
 
 	virtual ~FD3D12UniformBuffer();
 
+	virtual FRHIDescriptorHandle GetBindlessHandle() const override;
+
 	// Provides public non-const access to ResourceTable.
 	// @todo refactor uniform buffers to perform updates as a member function, so this isn't necessary.
 	TArray<TRefCountPtr<FRHIResource>>& GetResourceTable() { return ResourceTable; }
@@ -976,12 +976,12 @@ public:
 		return UpdateListeners.Num() != 0;
 	}
 
-	void UniformBufferUpdated(FRHICommandListBase& CmdList)
+	void UniformBufferUpdated(const FD3D12ContextArray& Contexts)
 	{
 		FScopeLock Lock(&UpdateListenersCS);
 		for (ID3D12UniformBufferUpdateListener* UpdateListener : UpdateListeners)
 		{
-			UpdateListener->UniformBufferUpdated(CmdList, this);
+			UpdateListener->UniformBufferUpdated(Contexts, this);
 		}
 	}
 };
diff --git a/Engine/Source/Runtime/D3D12RHI/Private/D3D12UniformBuffer.cpp b/Engine/Source/Runtime/D3D12RHI/Private/D3D12UniformBuffer.cpp
--- a/Engine/Source/Runtime/D3D12RHI/Private/D3D12UniformBuffer.cpp
+++ b/Engine/Source/Runtime/D3D12RHI/Private/D3D12UniformBuffer.cpp
@@ -21,6 +21,19 @@ inline bool AreBindlessUniformConstantsEnabled(FD3D12Device* Device)
 	return false;
 }
 
+inline bool DoesUniformBufferNeedBindlessHandle(FD3D12Device* Device, const FRHIUniformBufferLayout* Layout)
+{
+#if PLATFORM_SUPPORTS_BINDLESS_RENDERING
+	if (EnumHasAnyFlags(Layout->Flags, ERHIUniformBufferFlags::BindlessAccessible))
+	{
+		FD3D12BindlessDescriptorManager& Manager = Device->GetBindlessDescriptorManager();
+		return IsBindlessEnabledForAnyGraphics(Manager.GetConfiguration());
+	}
+#endif
+
+	return false;
+}
+
 FUniformBufferRHIRef FD3D12DynamicRHI::RHICreateUniformBuffer(const void* Contents, const FRHIUniformBufferLayout* Layout, EUniformBufferUsage Usage, EUniformBufferValidation Validation)
 {
 	SCOPE_CYCLE_COUNTER(STAT_D3D12UpdateUniformBufferTime);
@@ -47,10 +60,14 @@ FUniformBufferRHIRef FD3D12DynamicRHI::RHICreateUniformBuffer(const void* Conten
 			check(Align(Contents, 16) == Contents);
 			check(NumBytesActualData <= D3D12_REQ_CONSTANT_BUFFER_ELEMENT_COUNT * 16);
 
-#if D3D12RHI_USE_CONSTANT_BUFFER_VIEWS
-			// Create an offline CBV descriptor
-			NewUniformBuffer->View = new FD3D12ConstantBufferView(Device, FirstLinkedObject ? FirstLinkedObject->View : nullptr);
+			const bool bCreateBindlessHandle = DoesUniformBufferNeedBindlessHandle(Device, Layout);
+#if !D3D12RHI_USE_CONSTANT_BUFFER_VIEWS
+			if (bCreateBindlessHandle)
 #endif
+			{
+				// Create an offline CBV descriptor
+				NewUniformBuffer->View = new FD3D12ConstantBufferView(Device, FirstLinkedObject ? FirstLinkedObject->View : nullptr, bCreateBindlessHandle ? ERHIDescriptorType::CBV : ERHIDescriptorType::Invalid);
+			}
 
 			// Uniform buffers can be created without contents and updated later.
 			if (Contents)
@@ -75,9 +92,11 @@ FUniformBufferRHIRef FD3D12DynamicRHI::RHICreateUniformBuffer(const void* Conten
 
 				UE::RHICore::UpdateUniformBufferConstants(MappedData, Contents, *Layout, AreBindlessUniformConstantsEnabled(Device));
 
-#if D3D12RHI_USE_CONSTANT_BUFFER_VIEWS
-				NewUniformBuffer->View->CreateView(&NewUniformBuffer->ResourceLocation, 0, NumBytesActualData);
-#endif
+				if (NewUniformBuffer->View)
+				{
+					const uint32 NumBytes = Align(NumBytesActualData, D3D12_CONSTANT_BUFFER_DATA_PLACEMENT_ALIGNMENT);
+					NewUniformBuffer->View->CreateView(&NewUniformBuffer->ResourceLocation, 0, NumBytes);
+				}
 			}
 		}
 
@@ -111,51 +130,25 @@ FUniformBufferRHIRef FD3D12DynamicRHI::RHICreateUniformBuffer(const void* Conten
 	return UniformBufferOut;
 }
 
-FRHICOMMAND_MACRO(FRHICommandD3D12UpdateUniformBuffer)
+static void DoUpdateD3D12UniformBuffer(const FD3D12ContextArray& Contexts, FD3D12UniformBuffer* UniformBuffer, FD3D12ResourceLocation&& UpdatedLocation, TArrayView<FRHIResource*> UpdatedResources)
 {
-	TRefCountPtr<FD3D12UniformBuffer> UniformBuffer;
-	FD3D12ResourceLocation UpdatedLocation;
-	TArrayView<FRHIResource*> UpdatedResources;
-
-	FORCEINLINE_DEBUGGABLE FRHICommandD3D12UpdateUniformBuffer(FD3D12UniformBuffer* InUniformBuffer, FD3D12ResourceLocation& InUpdatedLocation, FRHIResource** InUpdatedResources, int32 InNumResources)
-		: UniformBuffer(InUniformBuffer)
-		, UpdatedLocation(InUpdatedLocation.GetParentDevice())
-		, UpdatedResources(InUpdatedResources, InNumResources)
+	for (int32 Index = 0; Index < UpdatedResources.Num(); Index++)
 	{
-		FD3D12ResourceLocation::TransferOwnership(UpdatedLocation, InUpdatedLocation);
-
-		for (FRHIResource* Resource : UpdatedResources)
-		{
-			Resource->AddRef();
-		}
-	}
+		UniformBuffer->GetResourceTable()[Index] = UpdatedResources[Index];
 
-	~FRHICommandD3D12UpdateUniformBuffer()
-	{
-		for (FRHIResource* Resource : UpdatedResources)
-		{
-			Resource->Release();
-		}
+		check(UniformBuffer->GetResourceTable()[Index]);
 	}
+	FD3D12ResourceLocation::TransferOwnership(UniformBuffer->ResourceLocation, UpdatedLocation);
 
-	void Execute(FRHICommandListBase& CmdList)
+	if (UniformBuffer->View)
 	{
-		for (int32 i = 0; i < UpdatedResources.Num(); ++i)
-		{
-			//check(UniformBuffer->ResourceTable[i]);
-			UniformBuffer->GetResourceTable()[i] = UpdatedResources[i];
-			check(UniformBuffer->GetResourceTable()[i]);
-		}
-		FD3D12ResourceLocation::TransferOwnership(UniformBuffer->ResourceLocation, UpdatedLocation);
-#if D3D12RHI_USE_CONSTANT_BUFFER_VIEWS
 		const uint32 NumBytes = Align(UniformBuffer->GetLayout().ConstantBufferSize, D3D12_CONSTANT_BUFFER_DATA_PLACEMENT_ALIGNMENT);
-		UniformBuffer->View->CreateView(&UniformBuffer->ResourceLocation, 0, NumBytes);
-#endif
-
-		// Notify the listeners now that the resource location on the uniform buffer has been updated
-		UniformBuffer->UniformBufferUpdated(CmdList);
+		UniformBuffer->View->UpdateView(Contexts, &UniformBuffer->ResourceLocation, 0, NumBytes);
 	}
-};
+
+	// Notify the listeners now that the resource location on the uniform buffer has been updated
+	UniformBuffer->UniformBufferUpdated(Contexts);
+}
 
 void FD3D12DynamicRHI::RHIUpdateUniformBuffer(FRHICommandListBase& RHICmdList, FRHIUniformBuffer* UniformBufferRHI, const void* Contents)
 {
@@ -169,17 +162,16 @@ void FD3D12DynamicRHI::RHIUpdateUniformBuffer(FRHICommandListBase& RHICmdList, F
 	FD3D12UniformBuffer* FirstUniformBuffer = ResourceCast(UniformBufferRHI);
 
 	const uint32 NumBytes = Layout.ConstantBufferSize;
-	const int32 NumResources = Layout.Resources.Num();
 
-	FRHIResource** CmdListResources = nullptr;
+	TArrayView<FRHIResource*> CmdListResources;
 
-	if (NumResources)
+	if (const int32 NumResources = Layout.Resources.Num())
 	{
 		CmdListResources = bBypass
-			? (FRHIResource**)FMemory_Alloca(sizeof(FRHIResource*) * NumResources) 
-			: (FRHIResource**)RHICmdList.Alloc(sizeof(FRHIResource*) * NumResources, alignof(FRHIResource*));
+			? TArrayView<FRHIResource*>((FRHIResource**)FMemory_Alloca(sizeof(FRHIResource*) * NumResources), NumResources)
+			: RHICmdList.AllocArrayUninitialized<FRHIResource*>(NumResources);
 
-		for (int32 Index = 0; Index < NumResources; ++Index)
+		for (int32 Index = 0; Index < NumResources; Index++)
 		{
 			const FRHIUniformBufferResource& Parameter = Layout.Resources[Index];
 			CmdListResources[Index] = GetShaderParameterResourceRHI(Contents, Parameter.MemberOffset, Parameter.MemberType);
@@ -189,7 +181,7 @@ void FD3D12DynamicRHI::RHIUpdateUniformBuffer(FRHICommandListBase& RHICmdList, F
 	// Update buffers on all GPUs by looping over FD3D12LinkedAdapterObject chain
 	for (FD3D12UniformBuffer& UniformBuffer : *FirstUniformBuffer)
 	{
-		check(UniformBuffer.GetResourceTable().Num() == NumResources);
+		check(UniformBuffer.GetResourceTable().Num() == CmdListResources.Num());
 
 		FD3D12Device* Device = UniformBuffer.GetParentDevice();
 
@@ -215,20 +207,29 @@ void FD3D12DynamicRHI::RHIUpdateUniformBuffer(FRHICommandListBase& RHICmdList, F
 			UE::RHICore::UpdateUniformBufferConstants(MappedData, Contents, Layout, AreBindlessUniformConstantsEnabled(Device));
 		}
 
-		if (bBypass)
+		if (!bBypass)
 		{
-			FRHICommandD3D12UpdateUniformBuffer Cmd(&UniformBuffer, UpdatedResourceLocation, CmdListResources, NumResources);
-			Cmd.Execute(RHICmdList);
+			for (FRHIResource* Resource : CmdListResources)
+			{
+				Resource->AddRef();
+			}
 		}
-		else
-		{
-			new (RHICmdList.AllocCommand<FRHICommandD3D12UpdateUniformBuffer>()) FRHICommandD3D12UpdateUniformBuffer(&UniformBuffer, UpdatedResourceLocation, CmdListResources, NumResources);
 
-			//fence is required to stop parallel recording threads from recording with the old bad state of the uniformbuffer resource table.  This command MUST execute before dependent recording starts.
-			RHICmdList.RHIThreadFence(true);
-		}
-	}
+		//fence is required to stop parallel recording threads from recording with the old bad state of the uniformbuffer resource table.  This command MUST execute before dependent recording starts.
+		RHICmdList.EnqueueLambdaMultiPipe(GetEnabledRHIPipelines(), FRHICommandListBase::EThreadFence::Enabled, TEXT("FD3D12DynamicRHI::RHIUpdateUniformBuffer"),
+			[bReleaseResources = !bBypass, UniformBuffer = &UniformBuffer, UpdatedLocation = MoveTemp(UpdatedResourceLocation), CmdListResources = MoveTemp(CmdListResources)](const FRHIContextArray& Contexts) mutable
+			{
+				DoUpdateD3D12UniformBuffer(Contexts, UniformBuffer, MoveTemp(UpdatedLocation), CmdListResources);
 
+				if (bReleaseResources)
+				{
+					for (FRHIResource* Resource : CmdListResources)
+					{
+						Resource->Release();
+					}
+				}
+			});
+	}
 }
 
 FD3D12UniformBuffer::~FD3D12UniformBuffer()
@@ -251,7 +252,10 @@ FD3D12UniformBuffer::~FD3D12UniformBuffer()
 	int64 BufferSize = ResourceLocation.GetSize();
 	DEC_MEMORY_STAT_BY(STAT_UniformBufferMemory, BufferSize);
 
-#if D3D12RHI_USE_CONSTANT_BUFFER_VIEWS
 	delete View;
-#endif
+}
+
+FRHIDescriptorHandle FD3D12UniformBuffer::GetBindlessHandle() const
+{
+	return View ? View->GetBindlessHandle() : FRHIDescriptorHandle();
 }
diff --git a/Engine/Source/Runtime/D3D12RHI/Private/D3D12View.cpp b/Engine/Source/Runtime/D3D12RHI/Private/D3D12View.cpp
--- a/Engine/Source/Runtime/D3D12RHI/Private/D3D12View.cpp
+++ b/Engine/Source/Runtime/D3D12RHI/Private/D3D12View.cpp
@@ -337,8 +337,8 @@ void FD3D12View::UpdateView(FD3D12ContextArray const& Contexts, const FResourceI
 //
 // -----------------------------------------------------------------------------------------------------
 
-FD3D12ConstantBufferView::FD3D12ConstantBufferView(FD3D12Device* InDevice, FD3D12ConstantBufferView* FirstLinkedObject)
-	: TD3D12View(InDevice, ERHIDescriptorType::CBV, ERHIDescriptorHeapType::Standard, FirstLinkedObject)
+FD3D12ConstantBufferView::FD3D12ConstantBufferView(FD3D12Device* InDevice, FD3D12ConstantBufferView* FirstLinkedObject, ERHIDescriptorType InDescriptorType)
+	: TD3D12View(InDevice, InDescriptorType, ERHIDescriptorHeapType::Standard, FirstLinkedObject)
 {
 }
 
@@ -353,6 +353,17 @@ void FD3D12ConstantBufferView::CreateView(FResourceInfo const& InResource, uint3
 	TD3D12View::CreateView(InResource, CBVDesc);
 }
 
+void FD3D12ConstantBufferView::UpdateView(const FD3D12ContextArray& Contexts, const FResourceInfo& InResource, uint32 InOffset, uint32 InAlignedSize)
+{
+	Offset = InOffset;
+
+	D3D12_CONSTANT_BUFFER_VIEW_DESC CBVDesc = {};
+	CBVDesc.BufferLocation = InResource.ResourceLocation->GetGPUVirtualAddress() + Offset;
+	CBVDesc.SizeInBytes = InAlignedSize;
+
+	TD3D12View::UpdateView(Contexts, InResource, CBVDesc);
+}
+
 void FD3D12ConstantBufferView::ResourceRenamed(FD3D12ContextArray const& Contexts, FD3D12BaseShaderResource* InRenamedResource, FD3D12ResourceLocation* InNewResourceLocation)
 {
 	D3DViewDesc.BufferLocation = InNewResourceLocation->GetGPUVirtualAddress() + Offset;
diff --git a/Engine/Source/Runtime/D3D12RHI/Private/D3D12View.h b/Engine/Source/Runtime/D3D12RHI/Private/D3D12View.h
--- a/Engine/Source/Runtime/D3D12RHI/Private/D3D12View.h
+++ b/Engine/Source/Runtime/D3D12RHI/Private/D3D12View.h
@@ -337,8 +337,9 @@ class FD3D12ConstantBufferView final : public TD3D12View<FD3D12ConstantBufferVie
 public:
 	static constexpr FD3D12OfflineDescriptor FD3D12DefaultViews::*Null { &FD3D12DefaultViews::NullCBV };
 
-	FD3D12ConstantBufferView(FD3D12Device* InParent, FD3D12ConstantBufferView* FirstLinkedObject);
-	void CreateView(FResourceInfo const& InResource, uint32 InOffset, uint32 InAlignedSize);
+	FD3D12ConstantBufferView(FD3D12Device* InParent, FD3D12ConstantBufferView* FirstLinkedObject, ERHIDescriptorType InDescriptorType = ERHIDescriptorType::Invalid);
+	void CreateView(const FResourceInfo& InResource, uint32 InOffset, uint32 InAlignedSize);
+	void UpdateView(const FD3D12ContextArray& Contexts, const FResourceInfo& InResource, uint32 InOffset, uint32 InAlignedSize);
 
 private:
 	virtual void ResourceRenamed(FD3D12ContextArray const& Contexts, FD3D12BaseShaderResource* InRenamedResource, FD3D12ResourceLocation* InNewResourceLocation) override;
diff --git a/Engine/Source/Runtime/RHI/Internal/RHIUniformBufferUtilities.h b/Engine/Source/Runtime/RHI/Internal/RHIUniformBufferUtilities.h
--- a/Engine/Source/Runtime/RHI/Internal/RHIUniformBufferUtilities.h
+++ b/Engine/Source/Runtime/RHI/Internal/RHIUniformBufferUtilities.h
@@ -51,6 +51,42 @@ inline void EnumerateUniformBufferResources(const FRHIUniformBuffer* RESTRICT Bu
 	}
 }
 
+#if PLATFORM_SUPPORTS_BINDLESS_RENDERING
+// Helper to call the respective setters for each FRHIResource contained inside the provided uniform buffer.
+template<typename TBinder>
+void SetAllUniformBufferResourcesForBindless(TBinder&& Binder, const FRHIUniformBuffer* Buffer)
+{
+	float CurrentTimeForTextureTimes = FApp::GetCurrentTime();
+
+	for (const TRefCountPtr<FRHIResource>& Resource : Buffer->GetResourceTable())
+	{
+		if (Resource)
+		{
+			const ERHIResourceType ResourceType = Resource->GetType();
+			switch (ResourceType)
+			{
+			case RRT_Texture:
+				static_cast<FRHITexture*>(Resource.GetReference())->SetLastRenderTime(CurrentTimeForTextureTimes);
+				Binder.SetTexture(static_cast<FRHITexture*>(Resource.GetReference()), 0);
+				break;
+			case RRT_ShaderResourceView:
+				Binder.SetSRV(static_cast<FRHIShaderResourceView*>(Resource.GetReference()), 0);
+				break;
+			case RRT_UnorderedAccessView:
+				Binder.SetUAV(static_cast<FRHIUnorderedAccessView*>(Resource.GetReference()), 0);
+				break;
+			case RRT_ResourceCollection:
+				Binder.SetResourceCollection(static_cast<FRHIResourceCollection*>(Resource.GetReference()), 0);
+				break;
+			case RRT_SamplerState:
+				Binder.SetSampler(static_cast<FRHISamplerState*>(Resource.GetReference()), 0);
+				break;
+			}
+		}
+	}
+}
+#endif
+
 template<typename TBinder, typename TUniformBufferArrayType, typename TBitMaskType, bool bFullyBindless = false>
 void SetUniformBufferResourcesFromTables(TBinder && Binder, FRHIShader const& Shader, TBitMaskType& DirtyUniformBuffers, TUniformBufferArrayType const& BoundUniformBuffers
                                          #if ENABLE_RHI_VALIDATION
diff --git a/Engine/Source/Runtime/RHI/Public/RHIDefinitions.h b/Engine/Source/Runtime/RHI/Public/RHIDefinitions.h
--- a/Engine/Source/Runtime/RHI/Public/RHIDefinitions.h
+++ b/Engine/Source/Runtime/RHI/Public/RHIDefinitions.h
@@ -715,6 +715,9 @@ enum class ERHIUniformBufferFlags : uint8
 
 	/** This struct is a view into uniform buffer object, on platforms that support UBO */
 	UniformView             = 1 << 3,
+
+	/** When bindless is enabled, make sure this uniform buffer can be accessed bindlessly. */
+	BindlessAccessible      = 1 << 4,
 };
 ENUM_CLASS_FLAGS(ERHIUniformBufferFlags);
 
diff --git a/Engine/Source/Runtime/RHI/Public/RHIResources.h b/Engine/Source/Runtime/RHI/Public/RHIResources.h
--- a/Engine/Source/Runtime/RHI/Public/RHIResources.h
+++ b/Engine/Source/Runtime/RHI/Public/RHIResources.h
@@ -1251,6 +1251,11 @@ public:
 
 	const TArray<TRefCountPtr<FRHIResource>>& GetResourceTable() const { return ResourceTable; }
 
+	virtual FRHIDescriptorHandle GetBindlessHandle() const
+	{
+		return FRHIDescriptorHandle();
+	}
+
 protected:
 	TArray<TRefCountPtr<FRHIResource>> ResourceTable;
 
diff --git a/Engine/Source/Runtime/RHI/Public/RHIShaderParameters.h b/Engine/Source/Runtime/RHI/Public/RHIShaderParameters.h
--- a/Engine/Source/Runtime/RHI/Public/RHIShaderParameters.h
+++ b/Engine/Source/Runtime/RHI/Public/RHIShaderParameters.h
@@ -294,10 +294,14 @@ struct FRHIBatchedShaderParameters
 
 	inline void SetShaderUniformBuffer(uint32 Index, FRHIUniformBuffer* UniformBuffer)
 	{
-		Allocator.Attach(this);
 		AddResourceParameter(UniformBuffer, (uint16)Index);
 	}
 
+	inline void SetBindlessShaderUniformBuffer(FRHIUniformBuffer* UniformBuffer)
+	{
+		AddBindlessParameter(UniformBuffer, (uint16)0);
+	}
+
 	inline void SetShaderTexture(uint32 Index, FRHITexture* Texture)
 	{
 		AddResourceParameter(Texture, (uint16)Index);
diff --git a/Engine/Source/Runtime/RenderCore/Public/ShaderParameterParser.h b/Engine/Source/Runtime/RenderCore/Public/ShaderParameterParser.h
--- a/Engine/Source/Runtime/RenderCore/Public/ShaderParameterParser.h
+++ b/Engine/Source/Runtime/RenderCore/Public/ShaderParameterParser.h
@@ -136,6 +136,7 @@ public:
 	// Prefix used to declare arrays of samplers/resources for bindless
 	static constexpr const TCHAR* kBindlessSRVArrayPrefix = TEXT("SRVDescriptorHeap_");
 	static constexpr const TCHAR* kBindlessUAVArrayPrefix = TEXT("UAVDescriptorHeap_");
+	static constexpr const TCHAR* kBindlessCBVArrayPrefix = TEXT("CBVDescriptorHeap_");
 	static constexpr const TCHAR* kBindlessSamplerArrayPrefix = TEXT("SamplerDescriptorHeap_");
 
 	static RENDERCORE_API EShaderParameterType ParseAndRemoveBindlessParameterPrefix(FStringView& InName);
diff --git a/Engine/Source/Runtime/VulkanRHI/Private/VulkanResources.h b/Engine/Source/Runtime/VulkanRHI/Private/VulkanResources.h
--- a/Engine/Source/Runtime/VulkanRHI/Private/VulkanResources.h
+++ b/Engine/Source/Runtime/VulkanRHI/Private/VulkanResources.h
@@ -994,7 +994,7 @@ public:
 		return UniformViewSRV != nullptr;
 	}
 
-	FRHIDescriptorHandle GetBindlessHandle() const
+	FRHIDescriptorHandle GetBindlessHandle() const override
 	{
 		checkfSlow(!BindlessHandle.IsValid() || (CachedDeviceAddress == GetDeviceAddress()), TEXT("Uniform Buffer bindless handle is out of date!"));
 		return BindlessHandle;
