DX12 Enhanced Barriers for Windows

Add support for the new layout, D3D12_LAYOUT_DIRECT_QUEUE_GENERIC_READ_COMPUTE_QUEUE_ACCESSIBLE, and prefer enhanced barriers on Windows when support is detected. This won't have any impact until we've updated to an Agility SDK that includes this new layout.

#Synced-CL 48187447

diff --git a/Engine/Source/Runtime/D3D12RHI/Private/D3D12Adapter.cpp b/Engine/Source/Runtime/D3D12RHI/Private/D3D12Adapter.cpp
--- a/Engine/Source/Runtime/D3D12RHI/Private/D3D12Adapter.cpp
+++ b/Engine/Source/Runtime/D3D12RHI/Private/D3D12Adapter.cpp
@@ -88,10 +88,17 @@ bool D3D12_ShouldBreakOnD3DDebugWarnings()
 	return GD3D12DebugCvar.GetValueOnAnyThread() == 3;
 }
 
-// 1 for now (e.g. EBs only on platforms where that's the only available implementation type)
-// There's a hole in the EB API that means we have to do some hacking to make it "work" and that's
-// not guaranteed to do the right thing on a wide range of drivers so we can't use them by default
-int32 GPreferredBarrierImplementation = 1;
+int32 GPreferredBarrierImplementation =
+#if D3D12RHI_HAS_LAYOUT_DIRECT_QUEUE_GENERIC_READ_COMPUTE_QUEUE_ACCESSIBLE
+	// Prefer enhanced barriers when available
+	2
+#else
+	// 1 for now (e.g. EBs only on platforms where that's the only available implementation type)
+	// There's a hole in the EB API that means we have to do some hacking to make it "work" and that's
+	// not guaranteed to do the right thing on a wide range of drivers so we can't use them by default
+	1
+#endif
+	;
 static FAutoConsoleVariableRef CVarPreferredBarrierImplementation(
 	TEXT("r.D3D12.PreferredBarrierImplementation"),
 	GPreferredBarrierImplementation,
diff --git a/Engine/Source/Runtime/D3D12RHI/Private/D3D12EnhancedBarriers.cpp b/Engine/Source/Runtime/D3D12RHI/Private/D3D12EnhancedBarriers.cpp
--- a/Engine/Source/Runtime/D3D12RHI/Private/D3D12EnhancedBarriers.cpp
+++ b/Engine/Source/Runtime/D3D12RHI/Private/D3D12EnhancedBarriers.cpp
@@ -27,7 +27,7 @@
 //         from multiple pipes (e.g. async and gfx) that includes access bits that are gfx specific, like 
 //         DEPTH_STENCIL_READ. Note that this makes the validation layers useless when the async pipe is used
 //         and it's questionable if any given driver will do something sane.
-#define PLATFORM_REQUIRES_ENHANCED_BARRIERS_GFX_ONLY_READ_BITS_HACK 1
+#define PLATFORM_REQUIRES_ENHANCED_BARRIERS_GFX_ONLY_READ_BITS_HACK !(D3D12RHI_HAS_LAYOUT_DIRECT_QUEUE_GENERIC_READ_COMPUTE_QUEUE_ACCESSIBLE)
 
 // @TODO - EB spec is inconsistent about whether D3D12_BARRIER_ACCESS_RAYTRACING_ACCELERATION_STRUCTURE_WRITE
 //         is compatible with D3D12_BARRIER_SYNC_RAYTRACING. Further, the validation layer in big windows
@@ -37,7 +37,7 @@
 // @TODO - Validation layer incorrectly complains that D3D12_BARRIER_ACCESS_DEPTH_STENCIL_READ isn't
 //         compatible with D3D12_BARRIER_LAYOUT_DEPTH_STENCIL_WRITE. Until this is fixed we omit the
 //         read access bit if we're also setting the write access bit
-#define PLATFORM_REQUIRES_LAYOUT_DEPTH_STENCIL_WRITE_NOT_COMPATIBLE_WITH_ACCESS_DEPTH_STENCIL_READ 1
+#define PLATFORM_REQUIRES_LAYOUT_DEPTH_STENCIL_WRITE_NOT_COMPATIBLE_WITH_ACCESS_DEPTH_STENCIL_READ !(PLATFORM_WINDOWS)
 
 // Debugging bits
 #define D3D12_ENHANCED_BARRIERS_LOG_BARRIERS_WHEN_BATCHED 0
@@ -359,6 +359,22 @@ static consteval auto GetIntersectionOf()
 	return ResizeArray<NumAndArray>();
 }
 
+// The below tables may be variably sized based on platform properties. This lets them 
+// vary in size without the ugly arithmetic necessary to explicitly calculate and specify
+// the size in the type based on those platform properties. It takes an array as input,
+// rather than the values themselves, to allow for a trailing comma.
+template <typename T, uint32 Size>
+static constexpr TStaticArray<T, Size> MakeStaticArray(const T(&InValues)[Size])
+{
+	// Used to capture the indices
+	return []
+		<uint32... Indices>
+		(const T(&InValues)[Size], TIntegerSequence<uint32, Indices...>)
+		{
+			return TStaticArray<T, Size>(InValues[Indices]...);
+		}(InValues, TMakeIntegerSequence<uint32, Size>{});
+}
+
 //
 // These tables are all copied directly from the DX spec
 // and are unmodified (except where specifically noted)
@@ -406,7 +422,7 @@ D3D12_BARRIER_SYNC ComputeQueueCompatibleSync =
 	| D3D12_BARRIER_SYNC_SPLIT;
 
 static constexpr
-TStaticArray<D3D12_BARRIER_SYNC, 23> AccessCompatibleSync = {
+TStaticArray AccessCompatibleSync = MakeStaticArray<D3D12_BARRIER_SYNC>({
 	//D3D12_BARRIER_ACCESS_VERTEX_BUFFER
 	D3D12_BARRIER_SYNC_ALL
 	| D3D12_BARRIER_SYNC_VERTEX_SHADING
@@ -549,15 +565,10 @@ TStaticArray<D3D12_BARRIER_SYNC, 23> AccessCompatibleSync = {
 
 	// Omitted since it's non-contigious with the other bits
 	//D3D12_BARRIER_ACCESS_NO_ACCESS
-};
-
-static constexpr
-uint32 LayoutCompatibleAccessNum = 31
-	+ !!(D3D12RHI_COMMON_AND_PRESENT_ARE_DISTINCT)
-	+ !!(D3D12RHI_HAS_LAYOUT_DIRECT_QUEUE_GENERIC_READ_COMPUTE_QUEUE_ACCESSIBLE);
+});
 
 static constexpr
-TStaticArray<D3D12_BARRIER_ACCESS, LayoutCompatibleAccessNum> LayoutCompatibleAccess = {
+TStaticArray LayoutCompatibleAccess = MakeStaticArray<D3D12_BARRIER_ACCESS>({
 	//D3D12_BARRIER_LAYOUT_COMMON
 	D3D12_BARRIER_ACCESS_SHADER_RESOURCE
 	| D3D12_BARRIER_ACCESS_COPY_DEST
@@ -681,9 +692,13 @@ TStaticArray<D3D12_BARRIER_ACCESS, LayoutCompatibleAccessNum> LayoutCompatibleAc
 
 #if D3D12RHI_HAS_LAYOUT_DIRECT_QUEUE_GENERIC_READ_COMPUTE_QUEUE_ACCESSIBLE
 	// D3D12_BARRIER_LAYOUT_DIRECT_QUEUE_GENERIC_READ_COMPUTE_QUEUE_ACCESSIBLE
-	D3D12_BARRIER_ACCESS_NO_ACCESS // @TODO - Fill this in after deducing from the debug layer!
+	D3D12_BARRIER_ACCESS_SHADER_RESOURCE
+	| D3D12_BARRIER_ACCESS_COPY_SOURCE
+	| D3D12_BARRIER_ACCESS_DEPTH_STENCIL_READ
+	| D3D12_BARRIER_ACCESS_SHADING_RATE_SOURCE
+	| D3D12_BARRIER_ACCESS_RESOLVE_SOURCE,
 #endif
-};
+});
 
 // These are not the same layout on some platforms but the compatible access should be the same
 static_assert(
@@ -749,7 +764,7 @@ D3D12_BARRIER_ACCESS AllQueueCompatibleAccess =
 	DirectQueueCompatibleAccess & ComputeQueueCompatibleAccess;
 
 static constexpr
-TStaticArray<D3D12_BARRIER_LAYOUT, 19> DirectQueueCompatibleLayouts = {
+TStaticArray DirectQueueCompatibleLayouts = MakeStaticArray<D3D12_BARRIER_LAYOUT>({
 	D3D12_BARRIER_LAYOUT_COMMON,
 	D3D12_BARRIER_LAYOUT_GENERIC_READ,
 	D3D12_BARRIER_LAYOUT_RENDER_TARGET,
@@ -762,30 +777,38 @@ TStaticArray<D3D12_BARRIER_LAYOUT, 19> DirectQueueCompatibleLayouts = {
 	D3D12_BARRIER_LAYOUT_RESOLVE_SOURCE,
 	D3D12_BARRIER_LAYOUT_RESOLVE_DEST,
 	D3D12_BARRIER_LAYOUT_SHADING_RATE_SOURCE,
+	// @TODO - Not listed in the spec, but logic tells us it's compatible
 	D3D12_BARRIER_LAYOUT_DIRECT_QUEUE_COMMON,
 	D3D12_BARRIER_LAYOUT_DIRECT_QUEUE_GENERIC_READ,
 	D3D12_BARRIER_LAYOUT_DIRECT_QUEUE_UNORDERED_ACCESS,
 	D3D12_BARRIER_LAYOUT_DIRECT_QUEUE_SHADER_RESOURCE,
 	D3D12_BARRIER_LAYOUT_DIRECT_QUEUE_COPY_SOURCE,
 	D3D12_BARRIER_LAYOUT_DIRECT_QUEUE_COPY_DEST,
+#if D3D12RHI_HAS_LAYOUT_DIRECT_QUEUE_GENERIC_READ_COMPUTE_QUEUE_ACCESSIBLE
+	D3D12_BARRIER_LAYOUT_DIRECT_QUEUE_GENERIC_READ_COMPUTE_QUEUE_ACCESSIBLE,
+#endif
 	// Not in spec
 	D3D12_BARRIER_LAYOUT_UNDEFINED,
-};
+});
 
 static constexpr
-TStaticArray<D3D12_BARRIER_LAYOUT, 14> ComputeQueueCompatibleLayouts = {
+TStaticArray ComputeQueueCompatibleLayouts = MakeStaticArray<D3D12_BARRIER_LAYOUT>({
 	D3D12_BARRIER_LAYOUT_COMMON,
 	D3D12_BARRIER_LAYOUT_GENERIC_READ,
 	D3D12_BARRIER_LAYOUT_UNORDERED_ACCESS,
 	D3D12_BARRIER_LAYOUT_SHADER_RESOURCE,
 	D3D12_BARRIER_LAYOUT_COPY_SOURCE,
 	D3D12_BARRIER_LAYOUT_COPY_DEST,
+	// @TODO - Not listed in the spec, but logic tells us it's compatible
 	D3D12_BARRIER_LAYOUT_COMPUTE_QUEUE_COMMON,
 	D3D12_BARRIER_LAYOUT_COMPUTE_QUEUE_GENERIC_READ,
 	D3D12_BARRIER_LAYOUT_COMPUTE_QUEUE_UNORDERED_ACCESS,
 	D3D12_BARRIER_LAYOUT_COMPUTE_QUEUE_SHADER_RESOURCE,
 	D3D12_BARRIER_LAYOUT_COMPUTE_QUEUE_COPY_SOURCE,
 	D3D12_BARRIER_LAYOUT_COMPUTE_QUEUE_COPY_DEST,
+#if D3D12RHI_HAS_LAYOUT_DIRECT_QUEUE_GENERIC_READ_COMPUTE_QUEUE_ACCESSIBLE
+	D3D12_BARRIER_LAYOUT_DIRECT_QUEUE_GENERIC_READ_COMPUTE_QUEUE_ACCESSIBLE,
+#endif
 	// Not in spec
 	D3D12_BARRIER_LAYOUT_UNDEFINED,
 #if PLATFORM_REQUIRES_ENHANCED_BARRIERS_GFX_ONLY_READ_BITS_HACK
@@ -796,7 +819,7 @@ TStaticArray<D3D12_BARRIER_LAYOUT, 14> ComputeQueueCompatibleLayouts = {
 	// when a resource is being read by multiple pipes at once
 	D3D12_BARRIER_LAYOUT_DIRECT_QUEUE_GENERIC_READ,
 #endif
-};
+});
 
 static constexpr
 TStaticArray AllQueueCompatibleLayouts =
@@ -1540,14 +1563,16 @@ static D3D12_BARRIER_LAYOUT GetEBLayout(
 			check(!EnumHasAnyFlags(InD3D12Access, GfxOnlyGenericReadBits));
 			return D3D12_BARRIER_LAYOUT_COMPUTE_QUEUE_GENERIC_READ;
 		case ERHIPipeline::All:
-#if PLATFORM_REQUIRES_ENHANCED_BARRIERS_GFX_ONLY_READ_BITS_HACK
-			// @TODO - This is to work around a hole in the API that won't allow DSVRead
-			//         access with any compute queue compatible layout.
 			if (EnumHasAnyFlags(InD3D12Access, GfxOnlyGenericReadBits))
 			{
+#if PLATFORM_REQUIRES_ENHANCED_BARRIERS_GFX_ONLY_READ_BITS_HACK
+				// @TODO - This is to work around a hole in the API that won't allow DSVRead
+				//         access with any compute queue compatible layout.
 				return D3D12_BARRIER_LAYOUT_DIRECT_QUEUE_GENERIC_READ;
-			}
+#elif D3D12RHI_HAS_LAYOUT_DIRECT_QUEUE_GENERIC_READ_COMPUTE_QUEUE_ACCESSIBLE
+				return D3D12_BARRIER_LAYOUT_DIRECT_QUEUE_GENERIC_READ_COMPUTE_QUEUE_ACCESSIBLE;
 #endif
+			}
 			return D3D12_BARRIER_LAYOUT_GENERIC_READ;
 		default:
 			checkNoEntry();
