Fix Simultaneous Access Resource Validation Errors

The debug layer in the new agility SDK (1.618.2) noticed the enhanced barriers weren't handling simultaneous access textures as they should. This fixes that as well as our own barrier validation logic.

#Synced-CL 48421968

diff --git a/Engine/Source/Runtime/D3D12RHI/Private/D3D12EnhancedBarriers.cpp b/Engine/Source/Runtime/D3D12RHI/Private/D3D12EnhancedBarriers.cpp
--- a/Engine/Source/Runtime/D3D12RHI/Private/D3D12EnhancedBarriers.cpp
+++ b/Engine/Source/Runtime/D3D12RHI/Private/D3D12EnhancedBarriers.cpp
@@ -281,6 +281,22 @@ static FString ConvertToString(D3D12_BARRIER_LAYOUT InLayout)
 		ENUMVAL(D3D12_BARRIER_LAYOUT_UNDEFINED));
 }
 
+static bool ResourceIsSimultaneousAccess(
+	const FD3D12ResourceDesc& InDesc)
+{
+	return EnumHasAnyFlags(InDesc.Flags, D3D12_RESOURCE_FLAG_ALLOW_SIMULTANEOUS_ACCESS);
+}
+
+static bool ResourceIsSimultaneousAccess(
+	const FD3D12Resource* InResource)
+{
+	if (InResource)
+	{
+		return ResourceIsSimultaneousAccess(InResource->GetDesc());
+	}
+	return false;
+}
+
 #undef ENUMVAL
 
 template <typename TArray>
@@ -567,6 +583,16 @@ TStaticArray AccessCompatibleSync = MakeStaticArray<D3D12_BARRIER_SYNC>({
 	//D3D12_BARRIER_ACCESS_NO_ACCESS
 });
 
+// Simultaneous access textures can be accessed using any access bits except these.
+// This info is rather hidden in the current version of the spec. It appears in the
+// section titled "Equivalent D3D12_BARRIER_ACCESS bit for each D3D12_RESOURCE_STATES bit"
+static constexpr
+D3D12_BARRIER_ACCESS SimultaneousAccessIncompatibleAccess =
+	D3D12_BARRIER_ACCESS_DEPTH_STENCIL_WRITE
+	| D3D12_BARRIER_ACCESS_DEPTH_STENCIL_READ
+	| D3D12_BARRIER_ACCESS_RAYTRACING_ACCELERATION_STRUCTURE_READ
+	| D3D12_BARRIER_ACCESS_RAYTRACING_ACCELERATION_STRUCTURE_WRITE;
+
 static constexpr
 TStaticArray LayoutCompatibleAccess = MakeStaticArray<D3D12_BARRIER_ACCESS>({
 	//D3D12_BARRIER_LAYOUT_COMMON
@@ -944,13 +970,19 @@ static bool AccessIsCompatibleWithQueue(
 
 static bool AccessIsCompatibleWithLayout(
 	D3D12_BARRIER_ACCESS InAccess,
-	D3D12_BARRIER_LAYOUT InLayout)
+	D3D12_BARRIER_LAYOUT InLayout,
+	const FD3D12Resource* InResource)
 {
 	if (InAccess == D3D12_BARRIER_ACCESS_NO_ACCESS)
 	{
 		return true;
 	}
 
+	if (ResourceIsSimultaneousAccess(InResource))
+	{
+		return !EnumHasAnyFlags(InAccess, SimultaneousAccessIncompatibleAccess);
+	}
+
 	if (static_cast<int32>(InLayout) >= LayoutCompatibleAccess.Num())
 	{
 		return false;
@@ -1110,8 +1142,8 @@ static bool CheckBarrierValuesAreCompatible(
 			check(LayoutIsCompatibleWithQueue(Barrier.LayoutBefore, InPipeline));
 			check(LayoutIsCompatibleWithQueue(Barrier.LayoutAfter, InPipeline));
 
-			check(AccessIsCompatibleWithLayout(Barrier.AccessBefore, Barrier.LayoutBefore));
-			check(AccessIsCompatibleWithLayout(Barrier.AccessAfter, Barrier.LayoutAfter));
+			check(AccessIsCompatibleWithLayout(Barrier.AccessBefore, Barrier.LayoutBefore, Barrier.pResource));
+			check(AccessIsCompatibleWithLayout(Barrier.AccessAfter, Barrier.LayoutAfter, Barrier.pResource));
 		}
 	}
 
@@ -1445,6 +1477,12 @@ static D3D12_BARRIER_LAYOUT GetEBLayout(
 		return D3D12_BARRIER_LAYOUT_UNDEFINED;
 	}
 
+	if (InTexture && ResourceIsSimultaneousAccess(InTexture->GetResource()))
+	{
+		// Simultaneous access textures have an immutable layout
+		return D3D12_BARRIER_LAYOUT_UNDEFINED;
+	}
+
 	// Special cases
 	if (EnumHasAnyFlags(InD3D12Access, ED3D12Access::CPURead | ED3D12Access::Common))
 	{
@@ -2225,10 +2263,17 @@ D3D12_BARRIER_LAYOUT FD3D12EnhancedBarriersForAdapterImpl::GetInitialLayout(
 	ED3D12Access InD3D12Access,
 	const FD3D12ResourceDesc& InDesc)
 {
-	// This makes the assumption that all resources begin life on the gfx pipe
 	const bool bIsBuffer = InDesc.Dimension == D3D12_RESOURCE_DIMENSION_BUFFER;
-	return bIsBuffer ?
-		D3D12_BARRIER_LAYOUT_UNDEFINED :
+	if (bIsBuffer)
+	{
+		// Buffers don't have a layout
+		return D3D12_BARRIER_LAYOUT_UNDEFINED;
+	}
+
+	return ResourceIsSimultaneousAccess(InDesc) ?
+		// Simultaneous access textures must be created in the COMMON layout
+		D3D12_BARRIER_LAYOUT_COMMON :
+		// This makes the assumption that all resources begin life on the gfx pipe
 		GetEBLayout(InD3D12Access, ERHIPipeline::Graphics, nullptr);
 }
 
diff --git a/Engine/Source/Runtime/D3D12RHI/Private/D3D12Texture.cpp b/Engine/Source/Runtime/D3D12RHI/Private/D3D12Texture.cpp
--- a/Engine/Source/Runtime/D3D12RHI/Private/D3D12Texture.cpp
+++ b/Engine/Source/Runtime/D3D12RHI/Private/D3D12Texture.cpp
@@ -796,12 +796,13 @@ FD3D12DynamicRHI::FCreateTextureInternalResult FD3D12DynamicRHI::CreateTextureIn
 
 	FD3D12TextureStats::D3D12TextureAllocated(*D3D12TextureOut);
 
-	// Discard for all RTV/DSV to be simple and safe.
+	// Discard for all non-simultaneous-access RTV/DSV to be simple and safe.
 	// THe debug layer reports these require Discard/Clear/Copy as the first operation:
 	// - Placed resources
 	// - Reserved resources
 	// - Committed resources with D3D12_HEAP_FLAG_CREATE_NOT_ZEROED flag
 	const bool bNeedsDiscard = (DesiredD3D12Access != ED3D12Access::Discard) &&
+		!EnumHasAnyFlags(ResourceDesc.Flags, D3D12_RESOURCE_FLAG_ALLOW_SIMULTANEOUS_ACCESS) &&
 		EnumHasAnyFlags(ResourceDesc.Flags, D3D12_RESOURCE_FLAG_ALLOW_RENDER_TARGET | D3D12_RESOURCE_FLAG_ALLOW_DEPTH_STENCIL);
 
 	FCreateTextureInternalResult Result{};
