Add More Flexible D3D12 Barrier Logging

Add a filter and make the logging runtime configurable via cvars

#Synced-CL 48489722

diff --git a/Engine/Source/Runtime/D3D12RHI/Private/D3D12EnhancedBarriers.cpp b/Engine/Source/Runtime/D3D12RHI/Private/D3D12EnhancedBarriers.cpp
--- a/Engine/Source/Runtime/D3D12RHI/Private/D3D12EnhancedBarriers.cpp
+++ b/Engine/Source/Runtime/D3D12RHI/Private/D3D12EnhancedBarriers.cpp
@@ -39,10 +39,6 @@
 //         read access bit if we're also setting the write access bit
 #define PLATFORM_REQUIRES_LAYOUT_DEPTH_STENCIL_WRITE_NOT_COMPATIBLE_WITH_ACCESS_DEPTH_STENCIL_READ !(PLATFORM_WINDOWS)
 
-// Debugging bits
-#define D3D12_ENHANCED_BARRIERS_LOG_BARRIERS_WHEN_BATCHED 0
-#define D3D12_ENHANCED_BARRIERS_LOG_BARRIERS_WHEN_FLUSHED 0
-
 
 // Each platform must provide its own implementation of this
 extern D3D12_BARRIER_LAYOUT GetSkipFastClearEliminateLayoutFlags();
@@ -1642,6 +1638,8 @@ static D3D12_BARRIER_LAYOUT GetEBLayout(
 	return GetEBLayout(ConvertToD3D12Access(InRHIAccess), InRHIPipe, InTexture);
 }
 
+#if D3D12RHI_BARRIER_LOGGING
+
 static void LogGlobalBarrier(
 	D3D12_BARRIER_SYNC SyncBefore,
 	D3D12_BARRIER_SYNC SyncAfter,
@@ -1667,41 +1665,48 @@ D3D12 Global Barrier
 #undef ENUM_STRING_AND_VALUE
 }
 
-template <typename TResourcePtr>
 static const ID3D12Resource* GetID3D12ResourcePtr(
-	TResourcePtr ResourcePtr)
+	const FD3D12Resource* ResourcePtr)
 {
-	using NonCVResourceType = std::remove_cv_t<std::remove_pointer_t<TResourcePtr>>;
+	return ResourcePtr->GetResource();
+}
 
-	if constexpr (std::is_same_v<FD3D12Resource, NonCVResourceType>)
-	{
-		return ResourcePtr->GetResource();
-	}
-	else if constexpr (std::is_same_v<ID3D12Resource, NonCVResourceType>)
-	{
-		return ResourcePtr;
-	}
+static const ID3D12Resource* GetID3D12ResourcePtr(
+	const ID3D12Resource* ResourcePtr)
+{
+	return ResourcePtr;
 }
 
-template <typename TResourcePtr>
 static FString GetResourceName(
-	const TResourcePtr ResourcePtr)
+	const FD3D12Resource* ResourcePtr)
 {
-	using NonCVResourceType = std::remove_cv_t<std::remove_pointer_t<TResourcePtr>>;
+	return ResourcePtr->GetName().ToString();
+}
 
-	if constexpr (std::is_same_v<FD3D12Resource, NonCVResourceType>)
-	{
-		return ResourcePtr->GetName()->ToString();
-	}
-	else if constexpr (std::is_same_v<ID3D12Resource, NonCVResourceType>)
-	{
-		return GetD312ObjectName(const_cast<NonCVResourceType*>(ResourcePtr));
-	}
+static FString GetResourceName(
+	const ID3D12Resource* ResourcePtr)
+{
+	return GetD312ObjectName(const_cast<ID3D12Resource*>(ResourcePtr));
+}
+
+static FString ConvertSubresourceToString(
+	uint32 Subresource)
+{
+	return (Subresource == D3D12_RESOURCE_BARRIER_ALL_SUBRESOURCES) ?
+		TEXT("All") :
+		FString::FromInt(Subresource);
+}
+
+static bool ShouldLogBarrier(
+	const FString& ResourceName)
+{
+	const FString& Filter = GD3D12LogBarriersResourceFilter;
+	return Filter.IsEmpty() || ResourceName.Contains(Filter);
 }
 
 template <typename TResourcePtr>
 static void LogTextureBarrier(
-	const TResourcePtr* pResource,
+	const TResourcePtr pResource,
 	D3D12_BARRIER_SYNC SyncBefore,
 	D3D12_BARRIER_SYNC SyncAfter,
 	D3D12_BARRIER_ACCESS AccessBefore,
@@ -1711,6 +1716,12 @@ static void LogTextureBarrier(
 	uint32 Subresource,
 	bool bDiscard)
 {
+	const FString ResourceName = GetResourceName(pResource);
+	if (!ShouldLogBarrier(ResourceName))
+	{
+		return;
+	}
+
 #define ENUM_STRING_AND_VALUE(EnumVal) \
 	*ConvertToString(EnumVal), static_cast<uint32>(EnumVal)
 
@@ -1724,17 +1735,17 @@ D3D12 Texture Barrier
   |     AccessAfter: %s (0x%08x)
   |    LayoutBefore: %s (0x%08x)
   |     LayoutAfter: %s (0x%08x)
-  |     Subresource: %u
+  |     Subresource: %s
   |         Discard: %u
 )"),
-	*GetResourceName(pResource), GetID3D12ResourcePtr(pResource),
+	*ResourceName, GetID3D12ResourcePtr(pResource),
 	ENUM_STRING_AND_VALUE(SyncBefore),
 	ENUM_STRING_AND_VALUE(SyncAfter),
 	ENUM_STRING_AND_VALUE(AccessBefore),
 	ENUM_STRING_AND_VALUE(AccessAfter),
 	ENUM_STRING_AND_VALUE(LayoutBefore),
 	ENUM_STRING_AND_VALUE(LayoutAfter),
-	Subresource,
+	*ConvertSubresourceToString(Subresource),
 	static_cast<uint32>(bDiscard));
 
 #undef ENUM_STRING_AND_VALUE
@@ -1742,12 +1753,18 @@ D3D12 Texture Barrier
 
 template <typename TResourcePtr>
 static void LogBufferBarrier(
-	const TResourcePtr* pResource,
+	const TResourcePtr pResource,
 	D3D12_BARRIER_SYNC SyncBefore,
 	D3D12_BARRIER_SYNC SyncAfter,
 	D3D12_BARRIER_ACCESS AccessBefore,
 	D3D12_BARRIER_ACCESS AccessAfter)
 {
+	const FString ResourceName = GetResourceName(pResource);
+	if (!ShouldLogBarrier(ResourceName))
+	{
+		return;
+	}
+
 #define ENUM_STRING_AND_VALUE(EnumVal) \
 	*ConvertToString(EnumVal), static_cast<uint32>(EnumVal)
 
@@ -1760,7 +1777,7 @@ D3D12 Buffer Barrier
   |    AccessBefore: %s (0x%08x)
   |     AccessAfter: %s (0x%08x)
 )"),
-	*GetResourceName(pResource), GetID3D12ResourcePtr(pResource),
+	*ResourceName, GetID3D12ResourcePtr(pResource),
 	ENUM_STRING_AND_VALUE(SyncBefore),
 	ENUM_STRING_AND_VALUE(SyncAfter),
 	ENUM_STRING_AND_VALUE(AccessBefore),
@@ -1838,6 +1855,8 @@ static void LogBarriers(
 	}
 }
 
+#endif // D3D12RHI_BARRIER_LOGGING
+
 // Use the top bit of the flags enum to mark transitions as "idle" time (used to remove the swapchain wait time for back buffers).
 static const D3D12_BARRIER_TYPE BarrierType_CountAsIdleTime = D3D12_BARRIER_TYPE(1ull << ((sizeof(D3D12_BARRIER_TYPE) * 8) - 1));
 static const D3D12_BARRIER_TYPE D3D12_BARRIER_TYPE_GLOBAL_COUNT_AS_IDLE = D3D12_BARRIER_TYPE(D3D12_BARRIER_TYPE_GLOBAL | BarrierType_CountAsIdleTime);
@@ -1963,8 +1982,11 @@ void FD3D12EnhancedBarriersBatcher::AddBarriers(
 
 	checkSlow(CheckBarrierValuesAreCompatible(Barriers, Context.GetRHIPipeline()));
 
-#if D3D12_ENHANCED_BARRIERS_LOG_BARRIERS_WHEN_BATCHED
-	LogBarriers(Barriers);
+#if D3D12RHI_BARRIER_LOGGING
+	if (GD3D12LogBarriers & D3D12LogBarriersWhenBatched)
+	{
+		LogBarriers(Barriers);
+	}
 #endif
 
 	static constexpr D3D12_BARRIER_TYPE BaseBarrierType = GetBarrierType<TNonQualifiedType>();
@@ -2242,8 +2264,11 @@ void FD3D12EnhancedBarriersBatcher::FlushIntoCommandList(
 		}
 
 		const int32 BatchNumGroups = BatchEnd - BatchStart;
-#if D3D12_ENHANCED_BARRIERS_LOG_BARRIERS_WHEN_FLUSHED
-		LogBarriers(MakeArrayView<const D3D12_BARRIER_GROUP>(&BarrierGroups[BatchStart], BatchNumGroups));
+#if D3D12RHI_BARRIER_LOGGING
+		if (GD3D12LogBarriers & D3D12LogBarriersWhenFlushed)
+		{
+			LogBarriers(MakeArrayView<const D3D12_BARRIER_GROUP>(&BarrierGroups[BatchStart], BatchNumGroups));
+		}
 #endif
 		CommandList.GraphicsCommandList8()->Barrier(BatchNumGroups, &BarrierGroups[BatchStart]);
 
diff --git a/Engine/Source/Runtime/D3D12RHI/Private/D3D12LegacyBarriers.cpp b/Engine/Source/Runtime/D3D12RHI/Private/D3D12LegacyBarriers.cpp
--- a/Engine/Source/Runtime/D3D12RHI/Private/D3D12LegacyBarriers.cpp
+++ b/Engine/Source/Runtime/D3D12RHI/Private/D3D12LegacyBarriers.cpp
@@ -102,18 +102,20 @@ static FString ConvertToResourceStateString(
 	return ResourceStateString;
 }
 
+#if D3D12RHI_BARRIER_LOGGING
+
 static void LogResourceBarriers(
 	TConstArrayView<D3D12_RESOURCE_BARRIER> Barriers,
 	ID3D12CommandList* const pCommandList,
-	ED3D12QueueType QueueType,
-	const FString& ResourceName)
+	ED3D12QueueType QueueType)
 {
 	// Configure what resource barriers are logged.
+	const FString& Filter = GD3D12LogBarriersResourceFilter;
 	const bool bLogAll = true;
 	const bool bLogTransitionDepth = true;
 	const bool bLogTransitionRenderTarget = true;
 	const bool bLogTransitionUAV = true;
-	const bool bCheckResourceName = ResourceName.IsEmpty() ? false : true;
+	const bool bCheckResourceName = !Filter.IsEmpty();
 
 	// Create the state bit mask to indicate what barriers should be logged.
 	uint32 ShouldLogMask = bLogAll ? static_cast<uint32>(-1) : 0;
@@ -147,7 +149,7 @@ static void LogResourceBarriers(
 			{
 				const FString BarrierResourceName = GetD312ObjectName(currentBarrier.Transition.pResource);
 
-				if ((bCheckResourceName == false) || (BarrierResourceName == ResourceName))
+				if (!bCheckResourceName || BarrierResourceName.Contains(Filter))
 				{
 					UE_LOG(LogD3D12RHI, Log, TEXT("*** BARRIER (CmdList: %016llX, Queue: %s) %u/%u: %s %016llX (Sub: %u), %s -> %s"), pCommandList, GetD3DCommandQueueTypeName(QueueType), i + 1, Barriers.Num(),
 						*BarrierResourceName,
@@ -163,7 +165,7 @@ static void LogResourceBarriers(
 		case D3D12_RESOURCE_BARRIER_TYPE_UAV:
 			{
 				const FString BarrierResourceName = GetD312ObjectName(currentBarrier.UAV.pResource);
-				if ((bCheckResourceName == false) || (BarrierResourceName == ResourceName))
+				if (!bCheckResourceName || BarrierResourceName.Contains(Filter))
 				{
 					UE_LOG(LogD3D12RHI, Log, TEXT("*** BARRIER (CmdList: %016llX, Queue: %s) %u/%u: UAV Barrier %s"), pCommandList, GetD3DCommandQueueTypeName(QueueType), i + 1, Barriers.Num(), *BarrierResourceName);
 				}
@@ -175,7 +177,7 @@ static void LogResourceBarriers(
 				const FString BarrierResourceNameBefore = GetD312ObjectName(currentBarrier.Aliasing.pResourceBefore);
 				const FString BarrierResourceNameAfter = GetD312ObjectName(currentBarrier.Aliasing.pResourceAfter);
 				
-				if ((bCheckResourceName == false) || (BarrierResourceNameBefore == ResourceName) || (BarrierResourceNameAfter == ResourceName))
+				if (!bCheckResourceName || BarrierResourceNameBefore.Contains(Filter) || BarrierResourceNameAfter.Contains(Filter))
 				{
 					UE_LOG(LogD3D12RHI, Log, TEXT("*** BARRIER (CmdList: %016llX, Queue: %s) %u/%u: Aliasing Barrier, %016llX %s -> %016llX %s"), pCommandList, GetD3DCommandQueueTypeName(QueueType), i + 1, Barriers.Num(), currentBarrier.Aliasing.pResourceBefore, *BarrierResourceNameBefore, currentBarrier.Aliasing.pResourceAfter, *BarrierResourceNameAfter);
 				}
@@ -189,6 +191,8 @@ static void LogResourceBarriers(
 	}
 }
 
+#endif // D3D12RHI_BARRIER_LOGGING
+
 static D3D12_RESOURCE_STATES GetDiscardedResourceState(
 	const D3D12_RESOURCE_DESC& InDesc,
 	ED3D12QueueType QueueType)
@@ -797,6 +801,16 @@ void FD3D12LegacyBarriersBatcher::AddUAV(
 	Barrier.Flags = D3D12_RESOURCE_BARRIER_FLAG_NONE;
 	Barrier.UAV.pResource = nullptr;	// Ignore the resource ptr for now. HW doesn't do anything with it.
 
+#if D3D12RHI_BARRIER_LOGGING
+	if (GD3D12LogBarriers & D3D12LogBarriersWhenBatched)
+	{
+		LogResourceBarriers(
+			MakeConstArrayView<D3D12_RESOURCE_BARRIER>({Barrier}),
+			Context.GetCommandList().BaseCommandList().Get(),
+			Context.GetCommandList().QueueType);
+	}
+#endif
+
 	if (!GD3D12BatchResourceBarriers)
 	{
 		FlushIntoCommandList(Context.GetCommandList(), Context.GetTimestampQueries());
@@ -841,6 +855,16 @@ int32 FD3D12LegacyBarriersBatcher::AddTransition(
 		Barrier.Flags |= BarrierFlag_CountAsIdleTime;
 	}
 
+#if D3D12RHI_BARRIER_LOGGING
+	if (GD3D12LogBarriers & D3D12LogBarriersWhenBatched)
+	{
+		LogResourceBarriers(
+			MakeConstArrayView<D3D12_RESOURCE_BARRIER>({Barrier}),
+			Context.GetCommandList().BaseCommandList().Get(),
+			Context.GetCommandList().QueueType);
+	}
+#endif
+
 	if (!GD3D12BatchResourceBarriers)
 	{
 		FlushIntoCommandList(Context.GetCommandList(), Context.GetTimestampQueries());
@@ -860,6 +884,16 @@ void FD3D12LegacyBarriersBatcher::AddAliasingBarrier(
 	Barrier.Aliasing.pResourceBefore = InResourceBefore;
 	Barrier.Aliasing.pResourceAfter = InResourceAfter;
 
+#if D3D12RHI_BARRIER_LOGGING
+	if (GD3D12LogBarriers & D3D12LogBarriersWhenBatched)
+	{
+		LogResourceBarriers(
+			MakeConstArrayView<D3D12_RESOURCE_BARRIER>({Barrier}),
+			Context.GetCommandList().BaseCommandList().Get(),
+			Context.GetCommandList().QueueType);
+	}
+#endif
+
 	if (!GD3D12BatchResourceBarriers)
 	{
 		FlushIntoCommandList(Context.GetCommandList(), Context.GetTimestampQueries());
@@ -903,10 +937,13 @@ void FD3D12LegacyBarriersBatcher::FlushIntoCommandList(
 			InsertTimestamp(true);
 		}
 
-#if DEBUG_RESOURCE_STATES
-   		TArrayView<FD3D12ResourceBarrier> SubsetBarriers(Barriers.GetData() + BatchStart, BatchEnd - BatchStart);
-   		TConstArrayView<D3D12_RESOURCE_BARRIER> ConstArrayView(	reinterpret_cast<const D3D12_RESOURCE_BARRIER*>(SubsetBarriers.GetData()), SubsetBarriers.Num());
-		LogResourceBarriers(ConstArrayView, CommandList.Interfaces.CommandList.GetReference(), CommandList.QueueType, FString(DX12_RESOURCE_NAME_TO_LOG));
+#if D3D12RHI_BARRIER_LOGGING
+		if (GD3D12LogBarriers & D3D12LogBarriersWhenFlushed)
+		{
+			TArrayView<FD3D12ResourceBarrier> SubsetBarriers(Barriers.GetData() + BatchStart, BatchEnd - BatchStart);
+			TConstArrayView<D3D12_RESOURCE_BARRIER> ConstArrayView(reinterpret_cast<const D3D12_RESOURCE_BARRIER*>(SubsetBarriers.GetData()), SubsetBarriers.Num());
+			LogResourceBarriers(ConstArrayView, CommandList.BaseCommandList().Get(), CommandList.QueueType);
+		}
 #endif
 		CommandList.GraphicsCommandList()->ResourceBarrier(BatchEnd - BatchStart, &Barriers[BatchStart]);
 #if DEBUG_RESOURCE_STATES
diff --git a/Engine/Source/Runtime/D3D12RHI/Private/D3D12RHIDefinitions.h b/Engine/Source/Runtime/D3D12RHI/Private/D3D12RHIDefinitions.h
--- a/Engine/Source/Runtime/D3D12RHI/Private/D3D12RHIDefinitions.h
+++ b/Engine/Source/Runtime/D3D12RHI/Private/D3D12RHIDefinitions.h
@@ -15,7 +15,6 @@
 #define RESIDENCY_PIPELINE_DEPTH	6
 
 #define DEBUG_RESOURCE_STATES	0
-#define DX12_RESOURCE_NAME_TO_LOG	""
 
 // Useful for debugging. Only flushes the default context though.
 #define D3D12RHI_IDLE_AFTER_EVERY_GPU_EVENT 0
@@ -46,6 +45,10 @@
 
 #define D3D12RHI_SUPPORTS_LEGACY_BARRIERS (D3D12RHI_ALLOW_LEGACY_BARRIERS)
 
+#ifndef D3D12RHI_BARRIER_LOGGING
+	#define D3D12RHI_BARRIER_LOGGING !(UE_BUILD_SHIPPING)
+#endif
+
 #if D3D12_MAX_DEVICE_INTERFACE >= 12
 	#define D3D12RHI_SUPPORTS_UNCOMPRESSED_UAV 1
 #else
diff --git a/Engine/Source/Runtime/D3D12RHI/Private/ID3D12Barriers.h b/Engine/Source/Runtime/D3D12RHI/Private/ID3D12Barriers.h
--- a/Engine/Source/Runtime/D3D12RHI/Private/ID3D12Barriers.h
+++ b/Engine/Source/Runtime/D3D12RHI/Private/ID3D12Barriers.h
@@ -2,6 +2,8 @@
 
 #pragma once
 
+#include "D3D12RHIDefinitions.h"
+
 // For the CVars
 #include "HAL/IConsoleManager.h"
 
@@ -45,6 +47,38 @@ inline FAutoConsoleVariableRef CVarD3D12BatchResourceBarriers(
 	GD3D12BatchResourceBarriers,
 	TEXT("Whether to allow batching resource barriers"));
 
+#if D3D12RHI_BARRIER_LOGGING
+inline int32 GD3D12LogBarriers = 0;
+inline FAutoConsoleVariableRef CVarD3D12LogBarriers(
+	TEXT("d3d12.LogBarriers"),
+	GD3D12LogBarriers,
+	TEXT("Whether to log barriers. Use d3d12.LogBarriersResourceFilter to filter by resource name.\n")
+	TEXT("0 - No logging\n")
+	TEXT("1 - Log when batched\n")
+	TEXT("2 - Log when flushed\n")
+	TEXT("3 - Log when batched and when flushed"));
+
+enum ED3D12LogBarriersFlags : decltype(GD3D12LogBarriers)
+{
+	D3D12LogBarriersNever       = 0,
+	D3D12LogBarriersWhenBatched = 1u << 0,
+	D3D12LogBarriersWhenFlushed = 1u << 1,
+};
+
+// Do not access the CVar directly as that will create a race.
+// Instead use this variable for which updates are sync'd via the callback
+inline FString GD3D12LogBarriersResourceFilter = "";
+inline FAutoConsoleVariable CVarD3D12LogBarriersResourceFilter(
+	TEXT("d3d12.LogBarriersResourceFilter"),
+	TEXT(""),
+	TEXT("Which resources to log barriers for. Set to empty to log everything."),
+	FConsoleVariableDelegate::CreateStatic(
+		[](IConsoleVariable* CVarFilter)
+		{
+			GD3D12LogBarriersResourceFilter = CVarFilter->GetString();
+		}));
+#endif // D3D12RHI_BARRIER_LOGGING
+
 class ID3D12BarriersForAdapter
 {
 public:
