Revert Change to Use DX12 EBs by Default

Testing with the new DirectX Agility SDK (v1.618.2) revealed additional bugs in the validation layer with Enhanced Barriers. These weren't present in 1.618.1. Since EBs now can't run cleanly with validation they can't be the default and so we'll again turn them off until the SDK is fixed. This will hopefully unblock us being able to upgrade our default Agility SDK version so we can test other aspects of it.

Note that this should have almost no practical impact as we've not yet updated our default Agility SDK to a version that would have turned on the EBs anyway.

#rnx
#Synced-CL 48544081

diff --git a/Engine/Source/Runtime/D3D12RHI/Private/D3D12Adapter.cpp b/Engine/Source/Runtime/D3D12RHI/Private/D3D12Adapter.cpp
--- a/Engine/Source/Runtime/D3D12RHI/Private/D3D12Adapter.cpp
+++ b/Engine/Source/Runtime/D3D12RHI/Private/D3D12Adapter.cpp
@@ -88,17 +88,10 @@ bool D3D12_ShouldBreakOnD3DDebugWarnings()
 	return GD3D12DebugCvar.GetValueOnAnyThread() == 3;
 }
 
-int32 GPreferredBarrierImplementation =
-#if D3D12RHI_HAS_LAYOUT_DIRECT_QUEUE_GENERIC_READ_COMPUTE_QUEUE_ACCESSIBLE
-	// Prefer enhanced barriers when available
-	2
-#else
-	// 1 for now (e.g. EBs only on platforms where that's the only available implementation type)
-	// There's a hole in the EB API that means we have to do some hacking to make it "work" and that's
-	// not guaranteed to do the right thing on a wide range of drivers so we can't use them by default
-	1
-#endif
-	;
+// 1 for now (e.g. EBs only on platforms where that's the only available implementation type)
+// There's a hole in the EB API that means we have to do some hacking to make it "work" and that's
+// not guaranteed to do the right thing on a wide range of drivers so we can't use them by default
+int32 GPreferredBarrierImplementation = 1;
 static FAutoConsoleVariableRef CVarPreferredBarrierImplementation(
 	TEXT("r.D3D12.PreferredBarrierImplementation"),
 	GPreferredBarrierImplementation,
