D3D12 Descriptor Cleanup
- Instead of passing around the full FRHIDescriptorHandle, we can just pass around the index. This cuts pending descriptor index storage in half.
- Removing an unused version of UE::D3D12Descriptors::CopyDescriptors
- Views and SamplerStates passed to InitializeDescriptor and UpdateDescriptor should be const

#Synced-CL 48667092

diff --git a/Engine/Source/Runtime/D3D12RHI/Private/D3D12BindlessDescriptors.cpp b/Engine/Source/Runtime/D3D12RHI/Private/D3D12BindlessDescriptors.cpp
--- a/Engine/Source/Runtime/D3D12RHI/Private/D3D12BindlessDescriptors.cpp
+++ b/Engine/Source/Runtime/D3D12RHI/Private/D3D12BindlessDescriptors.cpp
@@ -78,7 +78,7 @@ void FD3D12BindlessSamplerManager::CleanupResources()
 	GpuHeap = nullptr;
 }
 
-void FD3D12BindlessSamplerManager::InitializeDescriptor(FRHIDescriptorHandle DstHandle, FD3D12SamplerState* SamplerState)
+void FD3D12BindlessSamplerManager::InitializeDescriptor(FRHIDescriptorHandle DstHandle, const FD3D12SamplerState* SamplerState)
 {
 	check(DstHandle.GetType() == ERHIDescriptorType::Sampler);
 
@@ -415,7 +415,7 @@ void FD3D12BindlessResourceManager::GarbageCollect()
 				FGpuHeapData& GpuHeap = ActiveGpuHeaps[HeapIndex];
 				if (!GpuHeap.bInUse)
 				{
-					GpuHeap.UpdatedHandles.Empty();
+					GpuHeap.UpdatedIndices.Empty();
 					GpuHeap.LastUsedGarbageCollectCycle = GarbageCollectCycle;
 
 					PooledGpuHeaps.Add(GpuHeap);
@@ -483,7 +483,7 @@ void FD3D12BindlessResourceManager::Recycle(FD3D12DescriptorHeap* DescriptorHeap
 	}
 }
 
-void FD3D12BindlessResourceManager::InitializeDescriptor(FRHIDescriptorHandle DstHandle, FD3D12View* View)
+void FD3D12BindlessResourceManager::InitializeDescriptor(FRHIDescriptorHandle DstHandle, const FD3D12View* View)
 {
 	if (DstHandle.IsValid())
 	{
@@ -499,14 +499,14 @@ void FD3D12BindlessResourceManager::InitializeDescriptor(FRHIDescriptorHandle Ds
 		if (ActiveGpuHeapIndex >= 0 && !bCPUHeapResized)
 		{
 			UE::D3D12Descriptors::CopyDescriptor(GetParentDevice(), ActiveGpuHeaps[ActiveGpuHeapIndex].GpuHeap, DstHandle, OfflineCpuHandle);
-			ActiveGpuHeaps[ActiveGpuHeapIndex].UpdatedHandles.Add(DstHandle);
+			ActiveGpuHeaps[ActiveGpuHeapIndex].UpdatedIndices.Add(DstHandle.GetIndex());
 		}
 
 		INC_DWORD_STAT(STAT_D3D12BindlessResourceDescriptorsInitialized);
 	}
 }
 
-void FD3D12BindlessResourceManager::UpdateDescriptor(FD3D12ContextArray const& Contexts, FRHIDescriptorHandle DstHandle, FD3D12View* View)
+void FD3D12BindlessResourceManager::UpdateDescriptor(FD3D12ContextArray const& Contexts, FRHIDescriptorHandle DstHandle, const FD3D12View* View)
 {
 	if (DstHandle.IsValid())
 	{
@@ -533,7 +533,7 @@ void FD3D12BindlessResourceManager::UpdateDescriptor(FD3D12ContextArray const& C
 			}
 
 			bRequestNewActiveGpuHeap = true;
-			ActiveGpuHeaps[ActiveGpuHeapIndex].UpdatedHandles.Add(DstHandle);
+			ActiveGpuHeaps[ActiveGpuHeapIndex].UpdatedIndices.Add(DstHandle.GetIndex());
 		}
 
 		INC_DWORD_STAT(STAT_D3D12BindlessResourceDescriptorsUpdated);
@@ -706,7 +706,7 @@ void FD3D12BindlessResourceManager::CheckRequestNewActiveGPUHeap()
 		{
 			if (GpuHeapIndex != ActiveGpuHeapIndex)
 			{
-				ActiveGpuHeaps[GpuHeapIndex].UpdatedHandles.Append(ActiveGpuHeaps[ActiveGpuHeapIndex].UpdatedHandles);
+				ActiveGpuHeaps[GpuHeapIndex].UpdatedIndices.Append(ActiveGpuHeaps[ActiveGpuHeapIndex].UpdatedIndices);
 			}
 		}
 
@@ -731,10 +731,10 @@ void FD3D12BindlessResourceManager::CheckRequestNewActiveGPUHeap()
 		//		 CitySample there is about 2 to 4 times duplication but still faster to copy all then deduplication
 
 		FGpuHeapData& GpuHeapData = ActiveGpuHeaps[NewActiveGpuHeapIndex];
-		INC_DWORD_STAT_BY(STAT_D3D12BindlessResourceGPUDescriptorsCopied, GpuHeapData.UpdatedHandles.Num());
+		INC_DWORD_STAT_BY(STAT_D3D12BindlessResourceGPUDescriptorsCopied, GpuHeapData.UpdatedIndices.Num());
 
-		UE::D3D12Descriptors::CopyDescriptors(GetParentDevice(), GpuHeapData.GpuHeap, CpuHeap, GpuHeapData.UpdatedHandles);
-		GpuHeapData.UpdatedHandles.Reset();
+		UE::D3D12Descriptors::CopyDescriptors(GetParentDevice(), GpuHeapData.GpuHeap, CpuHeap, GpuHeapData.UpdatedIndices);
+		GpuHeapData.UpdatedIndices.Reset();
 
 		// Mark in use by GPU again
 		GpuHeapData.bInUse = true;
@@ -820,7 +820,7 @@ void FD3D12BindlessDescriptorManager::DeferredFreeFromDestructor(FRHIDescriptorH
 	}
 }
 
-void FD3D12BindlessDescriptorManager::InitializeDescriptor(FRHIDescriptorHandle DstHandle, FD3D12SamplerState* SamplerState)
+void FD3D12BindlessDescriptorManager::InitializeDescriptor(FRHIDescriptorHandle DstHandle, const FD3D12SamplerState* SamplerState)
 {
 	if (SamplerManager)
 	{
@@ -832,7 +832,7 @@ void FD3D12BindlessDescriptorManager::InitializeDescriptor(FRHIDescriptorHandle
 	checkNoEntry();
 }
 
-void FD3D12BindlessDescriptorManager::InitializeDescriptor(FRHIDescriptorHandle DstHandle, FD3D12View* View)
+void FD3D12BindlessDescriptorManager::InitializeDescriptor(FRHIDescriptorHandle DstHandle, const FD3D12View* View)
 {
 	if (ResourceManager && ResourceManager->HandlesAllocation(DstHandle.GetType()))
 	{
@@ -844,7 +844,7 @@ void FD3D12BindlessDescriptorManager::InitializeDescriptor(FRHIDescriptorHandle
 	checkNoEntry();
 }
 
-void FD3D12BindlessDescriptorManager::UpdateDescriptor(FD3D12ContextArray const& Contexts, FRHIDescriptorHandle DstHandle, FD3D12View* View)
+void FD3D12BindlessDescriptorManager::UpdateDescriptor(FD3D12ContextArray const& Contexts, FRHIDescriptorHandle DstHandle, const FD3D12View* View)
 {
 	if (ResourceManager)
 	{
diff --git a/Engine/Source/Runtime/D3D12RHI/Private/D3D12BindlessDescriptors.h b/Engine/Source/Runtime/D3D12RHI/Private/D3D12BindlessDescriptors.h
--- a/Engine/Source/Runtime/D3D12RHI/Private/D3D12BindlessDescriptors.h
+++ b/Engine/Source/Runtime/D3D12RHI/Private/D3D12BindlessDescriptors.h
@@ -86,7 +86,7 @@ public:
 
 	void CleanupResources();
 
-	void InitializeDescriptor(FRHIDescriptorHandle DstHandle, FD3D12SamplerState* SamplerState);
+	void InitializeDescriptor(FRHIDescriptorHandle DstHandle, const FD3D12SamplerState* SamplerState);
 
 	void OpenCommandList(FD3D12CommandContext& Context);
 	void CloseCommandList(FD3D12CommandContext& Context);
@@ -187,8 +187,8 @@ public:
 	void GarbageCollect();
 	void Recycle(FD3D12DescriptorHeap* DescriptorHeap);
 
-	void InitializeDescriptor(FRHIDescriptorHandle DstHandle, FD3D12View* View);
-	void UpdateDescriptor(FD3D12ContextArray const& Contexts, FRHIDescriptorHandle DstHandle, FD3D12View* View);
+	void InitializeDescriptor(FRHIDescriptorHandle DstHandle, const FD3D12View* View);
+	void UpdateDescriptor(FD3D12ContextArray const& Contexts, FRHIDescriptorHandle DstHandle, const FD3D12View* View);
 
 	void FlushPendingDescriptorUpdates(FD3D12CommandContext& Context);
 
@@ -241,10 +241,10 @@ private:
 
 	struct FGpuHeapData
 	{
-		FD3D12DescriptorHeapPtr      GpuHeap;
-		TArray<FRHIDescriptorHandle> UpdatedHandles;
-		bool                         bInUse = true;
-		uint64						 LastUsedGarbageCollectCycle = 0;
+		FD3D12DescriptorHeapPtr       GpuHeap;
+		TArray<FD3D12DescriptorIndex> UpdatedIndices;
+		bool                          bInUse = true;
+		uint64                        LastUsedGarbageCollectCycle = 0;
 	};		
 		
 	int32							ActiveGpuHeapIndex = -1;
@@ -296,9 +296,9 @@ public:
 	void GarbageCollect();
 	void Recycle(FD3D12DescriptorHeap* DescriptorHeap);
 
-	void InitializeDescriptor(FRHIDescriptorHandle DstHandle, FD3D12SamplerState* SamplerState);
-	void InitializeDescriptor(FRHIDescriptorHandle DstHandle, FD3D12View* View);
-	void UpdateDescriptor(FD3D12ContextArray const& Contexts, FRHIDescriptorHandle DstHandle, FD3D12View* SourceView);
+	void InitializeDescriptor(FRHIDescriptorHandle DstHandle, const FD3D12SamplerState* SamplerState);
+	void InitializeDescriptor(FRHIDescriptorHandle DstHandle, const FD3D12View* View);
+	void UpdateDescriptor(FD3D12ContextArray const& Contexts, FRHIDescriptorHandle DstHandle, const FD3D12View* SourceView);
 
 	void FinalizeContext(FD3D12CommandContext& Context);
 
diff --git a/Engine/Source/Runtime/D3D12RHI/Private/D3D12Descriptors.cpp b/Engine/Source/Runtime/D3D12RHI/Private/D3D12Descriptors.cpp
--- a/Engine/Source/Runtime/D3D12RHI/Private/D3D12Descriptors.cpp
+++ b/Engine/Source/Runtime/D3D12RHI/Private/D3D12Descriptors.cpp
@@ -5,14 +5,20 @@
 
 ////////////////////////////////////////////////////////////////////////////////////////////////////////////////////////
 
+void UE::D3D12Descriptors::CopyDescriptor(FD3D12Device* Device, FD3D12DescriptorHeap* TargetHeap, FD3D12DescriptorIndex DstIndex, D3D12_CPU_DESCRIPTOR_HANDLE SrcCpuHandle)
+{
+	const D3D12_CPU_DESCRIPTOR_HANDLE DestCpuHandle = TargetHeap->GetCPUSlotHandle(DstIndex);
+	const D3D12_DESCRIPTOR_HEAP_TYPE D3DHeapType = Translate(TargetHeap->GetType());
+
+	Device->GetDevice()->CopyDescriptorsSimple(1, DestCpuHandle, SrcCpuHandle, D3DHeapType);
+}
+
 void UE::D3D12Descriptors::CopyDescriptor(FD3D12Device* Device, FD3D12DescriptorHeap* TargetHeap, FRHIDescriptorHandle DstHandle, D3D12_CPU_DESCRIPTOR_HANDLE SrcCpuHandle)
 {
 	if (DstHandle.IsValid())
 	{
-		const D3D12_CPU_DESCRIPTOR_HANDLE DestCpuHandle = TargetHeap->GetCPUSlotHandle(DstHandle.GetIndex());
-		const D3D12_DESCRIPTOR_HEAP_TYPE D3DHeapType = Translate(DstHandle.GetType());
-
-		Device->GetDevice()->CopyDescriptorsSimple(1, DestCpuHandle, SrcCpuHandle, D3DHeapType);
+		checkSlow(Translate(TargetHeap->GetType()) == Translate(DstHandle.GetType()));
+		CopyDescriptor(Device, TargetHeap, DstHandle.GetIndex(), SrcCpuHandle);
 	}
 }
 
@@ -34,46 +40,13 @@ void UE::D3D12Descriptors::CopyDescriptors(FD3D12Device* Device, FD3D12Descripto
 	);
 }
 
-void UE::D3D12Descriptors::CopyDescriptors(FD3D12Device* Device, FD3D12DescriptorHeap* TargetHeap, FD3D12DescriptorHeap* SourceHeap, TConstArrayView<FRHIDescriptorHandle> Handles)
+void UE::D3D12Descriptors::CopyDescriptors(FD3D12Device* Device, FD3D12DescriptorHeap* TargetHeap, FD3D12DescriptorHeap* SourceHeap, TConstArrayView<FD3D12DescriptorIndex> Indices)
 {
 	// Use regular CPU Profile without dynamic value - almost same overhead for string formatting as actual Copy
 	TRACE_CPUPROFILER_EVENT_SCOPE_STR("CopyDescriptors HeapToHeap Scattered");
-	//SCOPED_NAMED_EVENT_F(TEXT("CopyDescriptors HeapToHeap Scattered (%d)"), FColor::Turquoise, Handles.Num());
-
-	const int32 NumDescriptors = Handles.Num();
-	const D3D12_DESCRIPTOR_HEAP_TYPE D3DHeapType = Translate(TargetHeap->GetType());
-
-	TArray<D3D12_CPU_DESCRIPTOR_HANDLE> DstDescriptors;
-	TArray<D3D12_CPU_DESCRIPTOR_HANDLE> SrcDescriptors;
-	TArray<UINT> CopyCounts;
-
-	DstDescriptors.SetNumUninitialized(NumDescriptors);
-	SrcDescriptors.SetNumUninitialized(NumDescriptors);
-	CopyCounts.SetNumUninitialized(NumDescriptors);
-
-	for (int32 Index = 0; Index < NumDescriptors; Index++)
-	{
-		DstDescriptors[Index] = TargetHeap->GetCPUSlotHandle(Handles[Index].GetIndex());
-		SrcDescriptors[Index] = SourceHeap->GetCPUSlotHandle(Handles[Index].GetIndex());
-		CopyCounts[Index] = 1;
-	}
-
-	Device->GetDevice()->CopyDescriptors(
-		NumDescriptors, DstDescriptors.GetData(), CopyCounts.GetData(),
-		NumDescriptors, SrcDescriptors.GetData(), CopyCounts.GetData(),
-		D3DHeapType
-	);
-}
-
-void UE::D3D12Descriptors::CopyDescriptors(FD3D12Device* Device, FD3D12DescriptorHeap* TargetHeap, TConstArrayView<FRHIDescriptorHandle> DstHandles, TConstArrayView<FD3D12OfflineDescriptor> SrcOfflineDescriptors)
-{
-	// Use regular CPU Profile without dynamic value - almost same overhead for string formatting as actual Copy
-	TRACE_CPUPROFILER_EVENT_SCOPE_STR("CopyDescriptors Scattered");
-	//SCOPED_NAMED_EVENT_F(TEXT("CopyDescriptors Scattered (%d)"), FColor::Turquoise, DstHandles.Num());
-
-	check(DstHandles.Num() == SrcOfflineDescriptors.Num());
+	//SCOPED_NAMED_EVENT_F(TEXT("CopyDescriptors HeapToHeap Scattered (%d)"), FColor::Turquoise, Indices.Num());
 
-	const int32 NumDescriptors = DstHandles.Num();
+	const int32 NumDescriptors = Indices.Num();
 	const D3D12_DESCRIPTOR_HEAP_TYPE D3DHeapType = Translate(TargetHeap->GetType());
 
 	TArray<D3D12_CPU_DESCRIPTOR_HANDLE> DstDescriptors;
@@ -86,8 +59,8 @@ void UE::D3D12Descriptors::CopyDescriptors(FD3D12Device* Device, FD3D12Descripto
 
 	for (int32 Index = 0; Index < NumDescriptors; Index++)
 	{
-		DstDescriptors[Index] = TargetHeap->GetCPUSlotHandle(DstHandles[Index].GetIndex());
-		SrcDescriptors[Index] = SrcOfflineDescriptors[Index];
+		DstDescriptors[Index] = TargetHeap->GetCPUSlotHandle(Indices[Index]);
+		SrcDescriptors[Index] = SourceHeap->GetCPUSlotHandle(Indices[Index]);
 		CopyCounts[Index] = 1;
 	}
 
diff --git a/Engine/Source/Runtime/D3D12RHI/Private/D3D12Descriptors.h b/Engine/Source/Runtime/D3D12RHI/Private/D3D12Descriptors.h
--- a/Engine/Source/Runtime/D3D12RHI/Private/D3D12Descriptors.h
+++ b/Engine/Source/Runtime/D3D12RHI/Private/D3D12Descriptors.h
@@ -78,13 +78,17 @@ inline D3D12_DESCRIPTOR_HEAP_FLAGS Translate(ED3D12DescriptorHeapFlags InHeapFla
 	return D3D12_DESCRIPTOR_HEAP_FLAG_NONE;
 }
 
+using FD3D12DescriptorIndex = uint32;
+
 namespace UE::D3D12Descriptors
 {
 	FD3D12DescriptorHeap* CreateDescriptorHeap(FD3D12Device* Device, const TCHAR* DebugName, ERHIDescriptorHeapType HeapType, uint32 NumDescriptors, ED3D12DescriptorHeapFlags Flags, bool bGlobal = false);
+
+	void CopyDescriptor(FD3D12Device* Device, FD3D12DescriptorHeap* TargetHeap, FD3D12DescriptorIndex DstIndex, D3D12_CPU_DESCRIPTOR_HANDLE SrcCpuHandle);
 	void CopyDescriptor(FD3D12Device* Device, FD3D12DescriptorHeap* TargetHeap, FRHIDescriptorHandle DstHandle, D3D12_CPU_DESCRIPTOR_HANDLE SrcCpuHandle);
+
 	void CopyDescriptors(FD3D12Device* Device, FD3D12DescriptorHeap* TargetHeap, FD3D12DescriptorHeap* SourceHeap, uint32 FirstDescriptor, uint32 NumDescriptors);
-	void CopyDescriptors(FD3D12Device* Device, FD3D12DescriptorHeap* TargetHeap, FD3D12DescriptorHeap* SourceHeap, TConstArrayView<FRHIDescriptorHandle> Handles);
-	void CopyDescriptors(FD3D12Device* Device, FD3D12DescriptorHeap* TargetHeap, TConstArrayView<FRHIDescriptorHandle> DstHandles, TConstArrayView<FD3D12OfflineDescriptor> SrcOfflineDescriptors);
+	void CopyDescriptors(FD3D12Device* Device, FD3D12DescriptorHeap* TargetHeap, FD3D12DescriptorHeap* SourceHeap, TConstArrayView<FD3D12DescriptorIndex> Indices);
 
 	FD3D12OfflineDescriptor CreateOfflineCopy(FD3D12Device* Device, D3D12_CPU_DESCRIPTOR_HANDLE InDescriptor, ERHIDescriptorHeapType InType);
 	FD3D12OfflineDescriptor CreateOfflineCopy(FD3D12Device* Device, FD3D12DescriptorHeap* InHeap, FRHIDescriptorHandle InHandle);
