Changing order of operation in FD3D12CommandContext::CloseCommandList in order to avoid FD3D12ContextCommon::CloseCommandList/FlushResourceBarriers to leave StateCache in an open state

#Synced-CL 48601921

diff --git a/Engine/Source/Runtime/D3D12RHI/Private/D3D12CommandContext.cpp b/Engine/Source/Runtime/D3D12RHI/Private/D3D12CommandContext.cpp
--- a/Engine/Source/Runtime/D3D12RHI/Private/D3D12CommandContext.cpp
+++ b/Engine/Source/Runtime/D3D12RHI/Private/D3D12CommandContext.cpp
@@ -375,9 +375,11 @@ void FD3D12ContextCommon::CloseCommandList()
 
 void FD3D12CommandContext::CloseCommandList()
 {
-	StateCache.GetDescriptorCache()->CloseCommandList();
+	// FD3D12ContextCommon::CloseCommandList call FlushResourceBarriers that might actually Open the command list
 	FD3D12ContextCommon::CloseCommandList();
 
+	StateCache.GetDescriptorCache()->CloseCommandList();
+
 #if PLATFORM_SUPPORTS_BINDLESS_RENDERING
 	// Always call the Bindless Manager CloseCommandList, it will determine when it needs to do anything.
 	GetParentDevice()->GetBindlessDescriptorManager().CloseCommandList(*this);
