Extend FD3D12Adapter::SetResidencyPriority to take optional resource size

This is used to prioritize between resources within the same coarse bucket, similar to how the default D3D12 resource priority algorithm works: https://learn.microsoft.com/en-us/windows/win32/direct3d12/residency#default-priority-algorithm
While the default D3D fine-grained prioritization is only done for resources in the high priority bucket, we can do it for all resources.

The idea is to make larger resources less likely to get demoted in over-committed VRAM situation, as it may be harder to promote them back due to physical video memory fragmentation.

PR #13975
#Synced-CL 48683013

diff --git a/Engine/Source/Runtime/D3D12RHI/Private/D3D12Adapter.cpp b/Engine/Source/Runtime/D3D12RHI/Private/D3D12Adapter.cpp
--- a/Engine/Source/Runtime/D3D12RHI/Private/D3D12Adapter.cpp
+++ b/Engine/Source/Runtime/D3D12RHI/Private/D3D12Adapter.cpp
@@ -2191,25 +2191,55 @@ void FD3D12Adapter::DumpTrackedAllocationData(FOutputDevice& OutputDevice, bool
 
 #endif // TRACK_RESOURCE_ALLOCATIONS
 
-void FD3D12Adapter::SetResidencyPriority(ID3D12Pageable* Pageable, D3D12_RESIDENCY_PRIORITY HeapPriority, uint32 GPUIndex)
+void FD3D12Adapter::SetResidencyPriority(ID3D12Pageable* Pageable, D3D12_RESIDENCY_PRIORITY HeapPriority, uint32 GPUIndex, TOptional<uint64> ResourceSize)
 {
-#if D3D12_RHI_RAYTRACING // ID3D12Device5 is currently tied to DXR support
+#if D3D12_MAX_DEVICE_INTERFACE >= 5
+
 	// GUID for our custom private data that holds the current priority
-	static const GUID DataGuid = GUID{ 0xB365961D, 0xA209, 0x4475, { 0x84, 0x9B, 0xDA, 0xFF, 0x90, 0x91, 0x2E, 0xB1 } };
+	checkf(GPUIndex <= 0xFF, TEXT("GPUIndex value must fit into a byte, as it's packed into private data ID"));
+
+	FD3D12Device* NodeDevice = GetDevice(GPUIndex);
+	ID3D12Device5* D3DDevice5 = NodeDevice->GetDevice5();
+	if (!D3DDevice5)
+	{
+		return;
+	}
+
+	const GUID DataGuid = GUID{ 0xB365961D, 0xA209, 0x4475, { 0x84, 0x9B, 0xDA, 0xFF, 0x90, 0x91, 0x2E, uint8(GPUIndex) } };
 
 	const uint32 DataSize = uint32(sizeof(HeapPriority));
 	uint32 ExistingDataSize = DataSize;
 	D3D12_RESIDENCY_PRIORITY ExistingPriority = (D3D12_RESIDENCY_PRIORITY)0;
 
-	HRESULT Result = Pageable->GetPrivateData(DataGuid, &ExistingDataSize, &ExistingPriority);
+	const HRESULT Result = Pageable->GetPrivateData(DataGuid, &ExistingDataSize, &ExistingPriority);
+	const bool bExistingDataValid = SUCCEEDED(Result) && ExistingDataSize == DataSize;
+
+	uint32 SizeGroup = 0;
 
-	if (!SUCCEEDED(Result) || ExistingDataSize != DataSize || ExistingPriority != HeapPriority)
+	if (ResourceSize.IsSet())
 	{
-		FD3D12Device* NodeDevice = GetDevice(GPUIndex);
-		NodeDevice->GetDevice5()->SetResidencyPriority(1, &Pageable, &HeapPriority);
+		// To further prioritize among resources within the same coarse bucket,
+		// the lowest 16-bits of the priority are set to the size of the resource divided by 10MB
+		// (saturating to 0xFFFF for very large heaps). This additional prioritization
+		// favors larger heaps and resources. This matches the prioritization done by default
+		// for high priority buckets by windows, except here it is done for all buckets (not just high priority).
+		// See https://learn.microsoft.com/en-us/windows/win32/direct3d12/residency#default-priority-algorithm
+		SizeGroup = uint32(FMath::Min<uint64>((*ResourceSize >> 20u) / 10u, 0xFFFFu));
+	}
+	else if (bExistingDataValid)
+	{
+		// Keep the previously set resource size bits if not explicitly provided
+		SizeGroup = uint32(ExistingPriority) & 0xFFFFu;
+	}
+
+	HeapPriority = D3D12_RESIDENCY_PRIORITY((uint32(HeapPriority) & 0xFFFF0000u) | SizeGroup);
+	
+	if (!bExistingDataValid || ExistingPriority != HeapPriority)
+	{
+		D3DDevice5->SetResidencyPriority(1, &Pageable, &HeapPriority);
 		Pageable->SetPrivateData(DataGuid, DataSize, &HeapPriority);
 	}
-#endif // D3D12_RHI_RAYTRACING
+#endif // D3D12_MAX_DEVICE_INTERFACE >= 5
 }
 
 void FD3D12Adapter::CreateTransition(
diff --git a/Engine/Source/Runtime/D3D12RHI/Private/D3D12Adapter.h b/Engine/Source/Runtime/D3D12RHI/Private/D3D12Adapter.h
--- a/Engine/Source/Runtime/D3D12RHI/Private/D3D12Adapter.h
+++ b/Engine/Source/Runtime/D3D12RHI/Private/D3D12Adapter.h
@@ -483,7 +483,19 @@ public:
 	};
 	void FindReleasedAllocationData(D3D12_GPU_VIRTUAL_ADDRESS InGPUVirtualAddress, TArray<FReleasedAllocationData>& OutAllocationData);
 
-	void SetResidencyPriority(ID3D12Pageable* Pageable, D3D12_RESIDENCY_PRIORITY HeapPriority, uint32 GPUIndex);
+	/**
+	 * Controls the heuristics used by Windows Video Memory Manager to demote resources from GPU physical memory when it is
+	 * overcommitted or fragmented. The D3D12 RHI performs access tracking and manual residency management, however it is
+	 * still possible to overcommit under some conditions on low-VRAM GPUs.
+	 * Render targets, depth buffers and UAVs are assigned high residency priority by default. 
+	 * To differentiate between high priority resources, the size is packed into low 16 bits of the priority value (in units of 10MB).
+	 * If ResourceSize parameter is specified, the final priority value will be packed in a mannner consistent with the
+	 * default priority assignment. Current priority value is stored in the Pageable private data. The size bits will be
+	 * preserved if priority changes but the size is not provided in subsequent calls to this function.
+	 * See https://learn.microsoft.com/en-us/windows/win32/direct3d12/residency.
+	*/
+	void SetResidencyPriority(ID3D12Pageable* Pageable, D3D12_RESIDENCY_PRIORITY HeapPriority,
+							  uint32 GPUIndex, TOptional<uint64> ResourceSize = {});
 
 #if PLATFORM_WINDOWS
 	HMODULE GetDxgiDllHandle() const { return DxgiDllHandle; };
diff --git a/Engine/Source/Runtime/D3D12RHI/Private/D3D12Device.cpp b/Engine/Source/Runtime/D3D12RHI/Private/D3D12Device.cpp
--- a/Engine/Source/Runtime/D3D12RHI/Private/D3D12Device.cpp
+++ b/Engine/Source/Runtime/D3D12RHI/Private/D3D12Device.cpp
@@ -183,22 +183,26 @@ ID3D12Device* FD3D12Device::GetDevice()
 	return GetParentAdapter()->GetD3DDevice();
 }
 
-#if D3D12_RHI_RAYTRACING
+#if D3D12_MAX_DEVICE_INTERFACE >= 5
 ID3D12Device5* FD3D12Device::GetDevice5()
 {
 	return GetParentAdapter()->GetD3DDevice5();
 }
+#endif //D3D12_MAX_DEVICE_INTERFACE >= 5
 
+#if D3D12_MAX_DEVICE_INTERFACE >= 7
 ID3D12Device7* FD3D12Device::GetDevice7()
 {
 	return GetParentAdapter()->GetD3DDevice7();
 }
+#endif //D3D12_MAX_DEVICE_INTERFACE >= 7
 
+#if D3D12_MAX_DEVICE_INTERFACE >= 9
 ID3D12Device9* FD3D12Device::GetDevice9()
 {
 	return GetParentAdapter()->GetD3DDevice9();
 }
-#endif // D3D12_RHI_RAYTRACING
+#endif // D3D12_MAX_DEVICE_INTERFACE >= 9
 
 #if D3D12_SUPPORTS_DXGI_DEBUG
 typedef HRESULT(WINAPI *FDXGIGetDebugInterface1)(UINT, REFIID, void **);
diff --git a/Engine/Source/Runtime/D3D12RHI/Private/D3D12Device.h b/Engine/Source/Runtime/D3D12RHI/Private/D3D12Device.h
--- a/Engine/Source/Runtime/D3D12RHI/Private/D3D12Device.h
+++ b/Engine/Source/Runtime/D3D12RHI/Private/D3D12Device.h
@@ -175,15 +175,23 @@ public:
 	void									  InitExplicitDescriptorHeap();
 	FD3D12ExplicitDescriptorHeapCache*		  GetExplicitDescriptorHeapCache() { return ExplicitDescriptorHeapCache; }
 
+#if D3D12_MAX_DEVICE_INTERFACE >= 5
+	ID3D12Device5*							  GetDevice5();
+#endif // D3D12_MAX_DEVICE_INTERFACE >= 5
+
+#if D3D12_MAX_DEVICE_INTERFACE >= 7
+	ID3D12Device7*							  GetDevice7();
+#endif // D3D12_MAX_DEVICE_INTERFACE >= 7
+
+#if D3D12_MAX_DEVICE_INTERFACE >= 9
+	ID3D12Device9*							  GetDevice9();
+#endif // D3D12_MAX_DEVICE_INTERFACE >= 9
+
 	// Ray Tracing
 #if D3D12_RHI_RAYTRACING
 	void									  InitRayTracing();
 	void									  CleanupRayTracing();
 
-	ID3D12Device5*							  GetDevice5();
-	ID3D12Device7*							  GetDevice7();
-	ID3D12Device9*							  GetDevice9();
-
 	FD3D12RayTracingPipelineCache*			  GetRayTracingPipelineCache           () { return RayTracingPipelineCache;            }
 	FD3D12Buffer*							  GetRayTracingDispatchRaysDescBuffer  (ED3D12QueueType QueueType) { return Queues[(uint32)QueueType].RayTracingDispatchRaysDescBuffer; }
 	FD3D12RayTracingCompactionRequestHandler* GetRayTracingCompactionRequestHandler() { return RayTracingCompactionRequestHandler; }
diff --git a/Engine/Source/Runtime/D3D12RHI/Private/D3D12PoolAllocator.cpp b/Engine/Source/Runtime/D3D12RHI/Private/D3D12PoolAllocator.cpp
--- a/Engine/Source/Runtime/D3D12RHI/Private/D3D12PoolAllocator.cpp
+++ b/Engine/Source/Runtime/D3D12RHI/Private/D3D12PoolAllocator.cpp
@@ -142,9 +142,8 @@ void FD3D12MemoryPool::Init()
 			ID3D12Pageable* HeapResource = AllocationStrategy == EResourceAllocationStrategy::kPlacedResource
 				? BackingHeap->GetHeap()
 				: BackingResource->GetPageable();
-			D3D12_RESIDENCY_PRIORITY HeapPriority = D3D12_RESIDENCY_PRIORITY_HIGH;
-			FD3D12Device* NodeDevice = Adapter->GetDevice(GPUIndex);
-			NodeDevice->GetDevice5()->SetResidencyPriority(1, &HeapResource, &HeapPriority);
+			const D3D12_RESIDENCY_PRIORITY HeapPriority = D3D12_RESIDENCY_PRIORITY_HIGH;
+			Adapter->SetResidencyPriority(HeapResource, HeapPriority, GPUIndex, PoolSize);
 		}
 	}
 #endif // D3D12_RHI_RAYTRACING
@@ -541,7 +540,7 @@ void FD3D12PoolAllocator::AllocateResource(
 		&& ResourceLocation.GetAllocatorType() != FD3D12ResourceLocation::AT_Pool)
 	{
 		ID3D12Pageable* HeapResource = ResourceLocation.GetResource()->GetPageable();
-		Adapter->SetResidencyPriority(HeapResource, D3D12_RESIDENCY_PRIORITY_HIGH, GPUIndex);
+		Adapter->SetResidencyPriority(HeapResource, D3D12_RESIDENCY_PRIORITY_HIGH, GPUIndex, InSize);
 	}
 #endif // D3D12_RHI_RAYTRACING
 }
diff --git a/Engine/Source/Runtime/D3D12RHI/Private/D3D12RayTracing.cpp b/Engine/Source/Runtime/D3D12RHI/Private/D3D12RayTracing.cpp
--- a/Engine/Source/Runtime/D3D12RHI/Private/D3D12RayTracing.cpp
+++ b/Engine/Source/Runtime/D3D12RHI/Private/D3D12RayTracing.cpp
@@ -3566,7 +3566,7 @@ static TRefCountPtr<FD3D12Buffer> CreateRayTracingBuffer(FD3D12Adapter* Adapter,
 		// Elevates the scratch buffer heap priority, which may help performance / stability in low memory conditions 
 		// (Acceleration structure already boosted from allocation side)
 		ID3D12Pageable* HeapResource = Result->GetResource()->GetPageable();
-		Adapter->SetResidencyPriority(HeapResource, D3D12_RESIDENCY_PRIORITY_HIGH, GPUIndex);
+		Adapter->SetResidencyPriority(HeapResource, D3D12_RESIDENCY_PRIORITY_HIGH, GPUIndex, Size);
 	}
 	else
 	{
diff --git a/Engine/Source/Runtime/D3D12RHI/Private/D3D12Resources.cpp b/Engine/Source/Runtime/D3D12RHI/Private/D3D12Resources.cpp
--- a/Engine/Source/Runtime/D3D12RHI/Private/D3D12Resources.cpp
+++ b/Engine/Source/Runtime/D3D12RHI/Private/D3D12Resources.cpp
@@ -520,7 +520,7 @@ void FD3D12Resource::CommitReservedResource(ID3D12CommandQueue* D3DCommandQueue,
 
 				if (bHighPriorityResource)
 				{
-					Adapter->SetResidencyPriority(D3DHeap, D3D12_RESIDENCY_PRIORITY_HIGH, GPUIndex);
+					Adapter->SetResidencyPriority(D3DHeap, D3D12_RESIDENCY_PRIORITY_HIGH, GPUIndex, NewHeapDesc.SizeInBytes);
 				}
 
 				TRefCountPtr<FD3D12Heap> NewHeap = new FD3D12Heap(GetParentDevice(), GetVisibilityMask());
diff --git a/Engine/Source/Runtime/D3D12RHI/Private/D3D12TransientResourceAllocator.cpp b/Engine/Source/Runtime/D3D12RHI/Private/D3D12TransientResourceAllocator.cpp
--- a/Engine/Source/Runtime/D3D12RHI/Private/D3D12TransientResourceAllocator.cpp
+++ b/Engine/Source/Runtime/D3D12RHI/Private/D3D12TransientResourceAllocator.cpp
@@ -64,7 +64,7 @@ FD3D12TransientHeap::FD3D12TransientHeap(const FInitializer& Initializer, FD3D12
 		TRefCountPtr<ID3D12Device5> D3DDevice5;
 		if (SUCCEEDED(D3DDevice->QueryInterface(IID_PPV_ARGS(D3DDevice5.GetInitReference()))))
 		{
-			Adapter->SetResidencyPriority(D3DHeap, D3D12_RESIDENCY_PRIORITY_HIGH, Device->GetGPUIndex());
+			Adapter->SetResidencyPriority(D3DHeap, D3D12_RESIDENCY_PRIORITY_HIGH, Device->GetGPUIndex(), Desc.SizeInBytes);
 		}
 #endif // PLATFORM_WINDOWS
 	}
