Use `UE::Private::FQueryableRefCountedObject` for places that are relying on reading the ref count.

#Synced-CL 48803015

diff --git a/Engine/Plugins/Runtime/GeometryCache/Source/GeometryCacheStreamer/Private/GeometryCacheStreamer.cpp b/Engine/Plugins/Runtime/GeometryCache/Source/GeometryCacheStreamer/Private/GeometryCacheStreamer.cpp
--- a/Engine/Plugins/Runtime/GeometryCache/Source/GeometryCacheStreamer/Private/GeometryCacheStreamer.cpp
+++ b/Engine/Plugins/Runtime/GeometryCache/Source/GeometryCacheStreamer/Private/GeometryCacheStreamer.cpp
@@ -51,7 +51,7 @@ private:
 private:
 	FTSTicker::FDelegateHandle TickHandle;
 
-	class FRefCountedStream : public FRefCountedObject
+	class FRefCountedStream : public UE::Private::FQueryableRefCountedObject
 	{
 	public:
 		FRefCountedStream(IGeometryCacheStream* InStream)
diff --git a/Engine/Source/Editor/UnrealEd/Private/Cooker/CookGenerationHelper.h b/Engine/Source/Editor/UnrealEd/Private/Cooker/CookGenerationHelper.h
--- a/Engine/Source/Editor/UnrealEd/Private/Cooker/CookGenerationHelper.h
+++ b/Engine/Source/Editor/UnrealEd/Private/Cooker/CookGenerationHelper.h
@@ -200,7 +200,7 @@ private:
  * reachability to fully support per-platform customization of the cook graph. See TODO_COOKGENERATIONHELPER for
  * code blocks where we implement the workarounds for this todo.
  */
-struct FGenerationHelper : public FRefCountedObject
+struct FGenerationHelper : public UE::Private::FQueryableRefCountedObject
 {
 public:
 	/** Store the provided CookPackageSplitter and prepare the packages to generate. */
diff --git a/Engine/Source/Runtime/CoreUObject/Public/Cooker/CookArtifact.h b/Engine/Source/Runtime/CoreUObject/Public/Cooker/CookArtifact.h
--- a/Engine/Source/Runtime/CoreUObject/Public/Cooker/CookArtifact.h
+++ b/Engine/Source/Runtime/CoreUObject/Public/Cooker/CookArtifact.h
@@ -65,7 +65,7 @@ namespace UE::Cook
  * Provides hooks to save global settings that invalidate the artifact when changed.
  * Provides hooks to clean the artifact when invalidated.
  */
-class ICookArtifact : public FRefCountedObject
+class ICookArtifact : public UE::Private::FQueryableRefCountedObject
 {
 public:
 	virtual ~ICookArtifact() = default;
diff --git a/Engine/Source/Runtime/D3D12RHI/Private/D3D12PipelineState.h b/Engine/Source/Runtime/D3D12RHI/Private/D3D12PipelineState.h
--- a/Engine/Source/Runtime/D3D12RHI/Private/D3D12PipelineState.h
+++ b/Engine/Source/Runtime/D3D12RHI/Private/D3D12PipelineState.h
@@ -296,7 +296,7 @@ struct FD3D12PipelineStateWorker : public FD3D12AdapterChild, public FNonAbandon
 	TRefCountPtr<ID3D12PipelineState> PSO;
 };
 
-struct FD3D12PipelineState : public FD3D12AdapterChild, public FD3D12MultiNodeGPUObject, public FNoncopyable, public FRefCountedObject
+struct FD3D12PipelineState : public FD3D12AdapterChild, public FD3D12MultiNodeGPUObject, public FNoncopyable, public UE::Private::FQueryableRefCountedObject
 {
 public:
 	explicit FD3D12PipelineState(FD3D12Adapter* Parent);
diff --git a/Engine/Source/Runtime/D3D12RHI/Private/D3D12Pipelinestate.cpp b/Engine/Source/Runtime/D3D12RHI/Private/D3D12Pipelinestate.cpp
--- a/Engine/Source/Runtime/D3D12RHI/Private/D3D12Pipelinestate.cpp
+++ b/Engine/Source/Runtime/D3D12RHI/Private/D3D12Pipelinestate.cpp
@@ -597,7 +597,8 @@ FD3D12PipelineState* FD3D12PipelineStateCacheBase::FindInLowLevelCache(const FD3
 			INC_DWORD_STAT(STAT_PSOGraphicsLowlevelCacheHit);
 			if (FD3D12PipelineState::UsePSORefCounting())
 			{
-				(*Found)->AddRef().CheckAtLeast(2);
+				const uint32 Refs = (*Found)->AddRef();
+				checkSlow(Refs >= 2);
 			}
 			return *Found;
 		}
@@ -640,7 +641,8 @@ void FD3D12PipelineStateCacheBase::RemoveFromLowLevelCache(FD3D12PipelineState*
 	else
 	{
 		// If another thread requested the pipeline state, we need to restore the refcount we just decremented
-		PipelineState->AddRef().CheckAtLeast(2);
+		const uint32 Refs = PipelineState->AddRef();
+		checkSlow(Refs >= 2);
 	}
 }
 
@@ -665,7 +667,8 @@ void FD3D12PipelineStateCacheBase::RemoveFromLowLevelCache(FD3D12PipelineState*
 	else
 	{
 		// If another thread requested the pipeline state, we need to restore the refcount we just decremented
-		PipelineState->AddRef().CheckAtLeast(2);
+		const uint32 Refs = PipelineState->AddRef();
+		checkSlow(Refs >= 2);
 	}
 }
 
@@ -748,7 +751,8 @@ FD3D12PipelineState* FD3D12PipelineStateCacheBase::FindInLowLevelCache(const FD3
 			INC_DWORD_STAT(STAT_PSOComputeLowlevelCacheHit);
 			if (FD3D12PipelineState::UsePSORefCounting())
 			{
-				(*Found)->AddRef().CheckAtLeast(2);
+				const uint32 Refs = (*Found)->AddRef();
+				checkSlow(Refs >= 2);
 			}
 			return *Found;
 		}
diff --git a/Engine/Source/Runtime/D3D12RHI/Private/D3D12Resources.h b/Engine/Source/Runtime/D3D12RHI/Private/D3D12Resources.h
--- a/Engine/Source/Runtime/D3D12RHI/Private/D3D12Resources.h
+++ b/Engine/Source/Runtime/D3D12RHI/Private/D3D12Resources.h
@@ -176,7 +176,7 @@ struct FD3D12ResourceDesc : public TZeroedStruct<FD3D12ResourceDesc>, D3D12_RESO
 #endif
 };
 
-class FD3D12Resource : public FRefCountedObject, public FD3D12DeviceChild, public FD3D12MultiNodeGPUObject
+class FD3D12Resource : public UE::Private::FQueryableRefCountedObject, public FD3D12DeviceChild, public FD3D12MultiNodeGPUObject
 {
 	friend class FD3D12Buffer;
 private:
diff --git a/Engine/Source/Runtime/Engine/Private/Streaming/TextureInstanceTask.h b/Engine/Source/Runtime/Engine/Private/Streaming/TextureInstanceTask.h
--- a/Engine/Source/Runtime/Engine/Private/Streaming/TextureInstanceTask.h
+++ b/Engine/Source/Runtime/Engine/Private/Streaming/TextureInstanceTask.h
@@ -21,7 +21,7 @@ enum ETaskState
 };
 
 template <typename TWork>
-class TDoWorkTask : public FRefCountedObject
+class TDoWorkTask : public UE::Private::FQueryableRefCountedObject
 {
 public:
 
diff --git a/Engine/Source/Runtime/Engine/Private/Streaming/TextureInstanceView.h b/Engine/Source/Runtime/Engine/Private/Streaming/TextureInstanceView.h
--- a/Engine/Source/Runtime/Engine/Private/Streaming/TextureInstanceView.h
+++ b/Engine/Source/Runtime/Engine/Private/Streaming/TextureInstanceView.h
@@ -105,7 +105,7 @@ typedef TArray<FStreamingViewInfoExtra, TInlineAllocator<4>> FStreamingViewInfoE
 // A constant view on the relationship between textures/meshes, components and bounds. 
 // Has everything needed for the worker task to compute the required size per texture/mesh.
 
-class FRenderAssetInstanceView : public FRefCountedObject
+class FRenderAssetInstanceView : public UE::Private::FQueryableRefCountedObject
 {
 public:
 
diff --git a/Engine/Source/Runtime/Engine/Public/Rendering/SkeletalMeshLODRenderData.h b/Engine/Source/Runtime/Engine/Public/Rendering/SkeletalMeshLODRenderData.h
--- a/Engine/Source/Runtime/Engine/Public/Rendering/SkeletalMeshLODRenderData.h
+++ b/Engine/Source/Runtime/Engine/Public/Rendering/SkeletalMeshLODRenderData.h
@@ -142,7 +142,7 @@ private:
 	TArray<FBoneIndexType> Indices;
 };
 
-class FSkeletalMeshLODRenderData : public FRefCountedObject
+class FSkeletalMeshLODRenderData : public UE::Private::FQueryableRefCountedObject
 {
 public:
 
diff --git a/Engine/Source/Runtime/Engine/Public/StaticMeshResources.h b/Engine/Source/Runtime/Engine/Public/StaticMeshResources.h
--- a/Engine/Source/Runtime/Engine/Public/StaticMeshResources.h
+++ b/Engine/Source/Runtime/Engine/Public/StaticMeshResources.h
@@ -422,7 +422,7 @@ struct TIsContiguousContainer<FStaticMeshSectionArray> : TIsContiguousContainer<
  * This structure is ref counted to allow the LOD streamer to evaluate the number of readers to it (readers that could access the CPU data).
  * Because the stream out clears the CPU readable data, CPU code that samples it must ensure to only reference LODs above CurrentFirstLODIdx.
  */
-struct FStaticMeshLODResources : public FRefCountedObject
+struct FStaticMeshLODResources : public UE::Private::FQueryableRefCountedObject
 {
 public:
 
diff --git a/Engine/Source/Runtime/Landscape/Classes/LandscapeHeightfieldCollisionComponent.h b/Engine/Source/Runtime/Landscape/Classes/LandscapeHeightfieldCollisionComponent.h
--- a/Engine/Source/Runtime/Landscape/Classes/LandscapeHeightfieldCollisionComponent.h
+++ b/Engine/Source/Runtime/Landscape/Classes/LandscapeHeightfieldCollisionComponent.h
@@ -99,7 +99,7 @@ public:
 	UFUNCTION(BlueprintCallable, Category = "Landscape")
 	LANDSCAPE_API ULandscapeComponent* GetRenderComponent() const;
 
-	struct FHeightfieldGeometryRef : public FRefCountedObject
+	struct FHeightfieldGeometryRef : public UE::Private::FQueryableRefCountedObject
 	{
 		FGuid Guid;
 
diff --git a/Engine/Source/Runtime/RenderCore/Public/RenderGraphResources.h b/Engine/Source/Runtime/RenderCore/Public/RenderGraphResources.h
--- a/Engine/Source/Runtime/RenderCore/Public/RenderGraphResources.h
+++ b/Engine/Source/Runtime/RenderCore/Public/RenderGraphResources.h
@@ -1186,7 +1186,7 @@ struct FRDGBufferUAVDesc final
 inline FRHIBufferCreateInfo Translate(const FRDGBufferDesc& InDesc);
 
 class FRDGPooledBuffer final
-	: public FRefCountedObject
+	: public UE::Private::FQueryableRefCountedObject
 {
 public:
 	FRDGPooledBuffer(FRHICommandListBase& RHICmdList, TRefCountPtr<FRHIBuffer> InBuffer, const FRDGBufferDesc& InDesc, uint32 InNumAllocatedElements, const TCHAR* InName)
