Removing the GetDeviceX shortcuts in favor of getting them from the adapter directly.

#Synced-CL 48916798

diff --git a/Engine/Source/Runtime/D3D12RHI/Private/D3D12Adapter.cpp b/Engine/Source/Runtime/D3D12RHI/Private/D3D12Adapter.cpp
--- a/Engine/Source/Runtime/D3D12RHI/Private/D3D12Adapter.cpp
+++ b/Engine/Source/Runtime/D3D12RHI/Private/D3D12Adapter.cpp
@@ -1575,7 +1575,7 @@ void FD3D12Adapter::InitializeRayTracing()
 #if D3D12_RHI_RAYTRACING
 	for (uint32 GPUIndex : FRHIGPUMask::All())
 	{
-		if (Devices[GPUIndex]->GetDevice5())
+		if (Devices[GPUIndex]->GetParentAdapter()->GetD3DDevice5())
 		{
 			Devices[GPUIndex]->InitRayTracing();
 		}
@@ -2199,7 +2199,7 @@ void FD3D12Adapter::SetResidencyPriority(ID3D12Pageable* Pageable, D3D12_RESIDEN
 	checkf(GPUIndex <= 0xFF, TEXT("GPUIndex value must fit into a byte, as it's packed into private data ID"));
 
 	FD3D12Device* NodeDevice = GetDevice(GPUIndex);
-	ID3D12Device5* D3DDevice5 = NodeDevice->GetDevice5();
+	ID3D12Device5* D3DDevice5 = NodeDevice->GetParentAdapter()->GetD3DDevice5();
 	if (!D3DDevice5)
 	{
 		return;
diff --git a/Engine/Source/Runtime/D3D12RHI/Private/D3D12Device.cpp b/Engine/Source/Runtime/D3D12RHI/Private/D3D12Device.cpp
--- a/Engine/Source/Runtime/D3D12RHI/Private/D3D12Device.cpp
+++ b/Engine/Source/Runtime/D3D12RHI/Private/D3D12Device.cpp
@@ -183,27 +183,6 @@ ID3D12Device* FD3D12Device::GetDevice()
 	return GetParentAdapter()->GetD3DDevice();
 }
 
-#if D3D12_MAX_DEVICE_INTERFACE >= 5
-ID3D12Device5* FD3D12Device::GetDevice5()
-{
-	return GetParentAdapter()->GetD3DDevice5();
-}
-#endif //D3D12_MAX_DEVICE_INTERFACE >= 5
-
-#if D3D12_MAX_DEVICE_INTERFACE >= 7
-ID3D12Device7* FD3D12Device::GetDevice7()
-{
-	return GetParentAdapter()->GetD3DDevice7();
-}
-#endif //D3D12_MAX_DEVICE_INTERFACE >= 7
-
-#if D3D12_MAX_DEVICE_INTERFACE >= 9
-ID3D12Device9* FD3D12Device::GetDevice9()
-{
-	return GetParentAdapter()->GetD3DDevice9();
-}
-#endif // D3D12_MAX_DEVICE_INTERFACE >= 9
-
 #if D3D12_SUPPORTS_DXGI_DEBUG
 typedef HRESULT(WINAPI *FDXGIGetDebugInterface1)(UINT, REFIID, void **);
 #endif
diff --git a/Engine/Source/Runtime/D3D12RHI/Private/D3D12Device.h b/Engine/Source/Runtime/D3D12RHI/Private/D3D12Device.h
--- a/Engine/Source/Runtime/D3D12RHI/Private/D3D12Device.h
+++ b/Engine/Source/Runtime/D3D12RHI/Private/D3D12Device.h
@@ -175,18 +175,6 @@ public:
 	void									  InitExplicitDescriptorHeap();
 	FD3D12ExplicitDescriptorHeapCache*		  GetExplicitDescriptorHeapCache() { return ExplicitDescriptorHeapCache; }
 
-#if D3D12_MAX_DEVICE_INTERFACE >= 5
-	ID3D12Device5*							  GetDevice5();
-#endif // D3D12_MAX_DEVICE_INTERFACE >= 5
-
-#if D3D12_MAX_DEVICE_INTERFACE >= 7
-	ID3D12Device7*							  GetDevice7();
-#endif // D3D12_MAX_DEVICE_INTERFACE >= 7
-
-#if D3D12_MAX_DEVICE_INTERFACE >= 9
-	ID3D12Device9*							  GetDevice9();
-#endif // D3D12_MAX_DEVICE_INTERFACE >= 9
-
 	// Ray Tracing
 #if D3D12_RHI_RAYTRACING
 	void									  InitRayTracing();
diff --git a/Engine/Source/Runtime/D3D12RHI/Private/D3D12RayTracing.cpp b/Engine/Source/Runtime/D3D12RHI/Private/D3D12RayTracing.cpp
--- a/Engine/Source/Runtime/D3D12RHI/Private/D3D12RayTracing.cpp
+++ b/Engine/Source/Runtime/D3D12RHI/Private/D3D12RayTracing.cpp
@@ -1036,7 +1036,7 @@ public:
 			: Entry(InEntry)
 			, CacheKey(InCacheKey)
 			, Device(InDevice)
-			, RayTracingDevice(InDevice->GetDevice5())
+			, RayTracingDevice(InDevice->GetParentAdapter()->GetD3DDevice5())
 			, CollectionType(InCollectionType)
 			, bRequired(bInRequired)
 		{
@@ -2738,7 +2738,7 @@ FD3D12RayTracingPipelineState::FD3D12RayTracingPipelineState(FD3D12Device* Devic
 
 	TotalCreationTime -= FPlatformTime::Cycles64();
 
-	ID3D12Device5* RayTracingDevice = Device->GetDevice5();
+	ID3D12Device5* RayTracingDevice = Device->GetParentAdapter()->GetD3DDevice5();
 
 	TArrayView<FRHIRayTracingShader*> InitializerHitGroups = Initializer.GetHitGroupTable();
 	TArrayView<FRHIRayTracingShader*> InitializerMissShaders = Initializer.GetMissTable();
@@ -2967,7 +2967,7 @@ FD3D12RayTracingPipelineState::FD3D12RayTracingPipelineState(FD3D12Device* Devic
 			Desc.NumSubobjects = Subobjects.Num();
 			Desc.pSubobjects = Subobjects.GetData();
 
-			ID3D12Device7* Device7 = Device->GetDevice7();
+			ID3D12Device7* Device7 = Device->GetParentAdapter()->GetD3DDevice7();
 
 			VERIFYD3D12RESULT(Device7->AddToStateObject(&Desc,
 				BasePipeline->StateObject.GetReference(),
diff --git a/Engine/Source/Runtime/D3D12RHI/Private/Windows/WindowsD3D12Device.cpp b/Engine/Source/Runtime/D3D12RHI/Private/Windows/WindowsD3D12Device.cpp
--- a/Engine/Source/Runtime/D3D12RHI/Private/Windows/WindowsD3D12Device.cpp
+++ b/Engine/Source/Runtime/D3D12RHI/Private/Windows/WindowsD3D12Device.cpp
@@ -1981,7 +1981,7 @@ void FD3D12Device::GetRaytracingAccelerationStructurePrebuildInfo(
 	const D3D12_BUILD_RAYTRACING_ACCELERATION_STRUCTURE_INPUTS* pDesc,
 	D3D12_RAYTRACING_ACCELERATION_STRUCTURE_PREBUILD_INFO* pInfo)
 {
-	ID3D12Device5* RayTracingDevice = GetDevice5();
+	ID3D12Device5* RayTracingDevice = GetParentAdapter()->GetD3DDevice5();
 	RayTracingDevice->GetRaytracingAccelerationStructurePrebuildInfo(pDesc, pInfo);
 }
 
