More specific check with IsResourceTrackingEnabled before actually trying to use Nvidia Aftermath resource tracking

#Synced-CL 49399013

diff --git a/Engine/Source/Runtime/D3D12RHI/Private/D3D12NvidiaExtensions.cpp b/Engine/Source/Runtime/D3D12RHI/Private/D3D12NvidiaExtensions.cpp
--- a/Engine/Source/Runtime/D3D12RHI/Private/D3D12NvidiaExtensions.cpp
+++ b/Engine/Source/Runtime/D3D12RHI/Private/D3D12NvidiaExtensions.cpp
@@ -257,7 +257,7 @@ namespace UE::RHICore::Nvidia::Aftermath::D3D12
 	FResource RegisterResource(ID3D12Resource* D3DResource)
 	{
 		FResource Handle{};
-		if (IsEnabled())
+		if (IsResourceTrackingEnabled())
 		{
 			GFSDK_Aftermath_Result Result = GFSDK_Aftermath_DX12_RegisterResource(D3DResource, &Handle);
 			if (Result != GFSDK_Aftermath_Result_Success)
@@ -271,7 +271,7 @@ namespace UE::RHICore::Nvidia::Aftermath::D3D12
 
 	void UnregisterResource(FResource Resource)
 	{
-		if (IsEnabled() && Resource)
+		if (IsResourceTrackingEnabled() && Resource)
 		{
 			GFSDK_Aftermath_Result Result = GFSDK_Aftermath_DX12_UnregisterResource(Resource);
 			UE_CLOG(Result != GFSDK_Aftermath_Result_Success, LogD3D12RHI, VeryVerbose, TEXT("GFSDK_Aftermath_DX12_UnregisterResource failed: 0x%08x"), Result);
