Remove unused D3D12 RayTracing allocations when feature level is not SM6

Verified the NvAPI change saves about 100 MB of CPU memory when using ES31. The other removals are tiny, but they stop ETL from blaming D3D12 RayTracing for a 32 MB pooled allocation.

#Synced-CL 49539332

diff --git a/Engine/Source/Runtime/D3D12RHI/Private/D3D12Device.cpp b/Engine/Source/Runtime/D3D12RHI/Private/D3D12Device.cpp
--- a/Engine/Source/Runtime/D3D12RHI/Private/D3D12Device.cpp
+++ b/Engine/Source/Runtime/D3D12RHI/Private/D3D12Device.cpp
@@ -383,24 +383,27 @@ void FD3D12Device::SetupAfterDeviceCreation()
 	UpdateMSAASettings();
 
 #if D3D12_RHI_RAYTRACING
-	check(RayTracingCompactionRequestHandler == nullptr);
-	RayTracingCompactionRequestHandler = new FD3D12RayTracingCompactionRequestHandler(this);
-	for (FD3D12Queue& Queue : Queues)
+	if (GMaxRHIFeatureLevel >= ERHIFeatureLevel::SM6)
 	{
-		const size_t DispatchRaysSize = sizeof(D3D12_DISPATCH_RAYS_DESC);
-
-		const D3D12_RESOURCE_DESC DispatchRaysBufferDesc = CD3DX12_RESOURCE_DESC::Buffer(DispatchRaysSize, D3D12RHI_RESOURCE_FLAG_ALLOW_INDIRECT_BUFFER | D3D12_RESOURCE_FLAG_ALLOW_UNORDERED_ACCESS);
-		const FRHIBufferCreateDesc CreateDesc =
-			FRHIBufferCreateDesc::Create(TEXT("DispatchRaysDescBuffer"), DispatchRaysSize, 0, EBufferUsageFlags::DrawIndirect)
-			.SetGPUMask(GetGPUMask())
-			.SetInitialState(ERHIAccess::IndirectArgs);
-
-		Queue.RayTracingDispatchRaysDescBuffer = GetParentAdapter()->CreateRHIBuffer(
-			DispatchRaysBufferDesc,
-			256,
-			CreateDesc,
-			ED3D12ResourceStateMode::MultiState,
-			ED3D12Access::IndirectArgs);
+		check(RayTracingCompactionRequestHandler == nullptr);
+		RayTracingCompactionRequestHandler = new FD3D12RayTracingCompactionRequestHandler(this);
+		for (FD3D12Queue& Queue : Queues)
+		{
+			const size_t DispatchRaysSize = sizeof(D3D12_DISPATCH_RAYS_DESC);
+
+			const D3D12_RESOURCE_DESC DispatchRaysBufferDesc = CD3DX12_RESOURCE_DESC::Buffer(DispatchRaysSize, D3D12RHI_RESOURCE_FLAG_ALLOW_INDIRECT_BUFFER | D3D12_RESOURCE_FLAG_ALLOW_UNORDERED_ACCESS);
+			const FRHIBufferCreateDesc CreateDesc =
+				FRHIBufferCreateDesc::Create(TEXT("DispatchRaysDescBuffer"), DispatchRaysSize, 0, EBufferUsageFlags::DrawIndirect)
+				.SetGPUMask(GetGPUMask())
+				.SetInitialState(ERHIAccess::IndirectArgs);
+
+			Queue.RayTracingDispatchRaysDescBuffer = GetParentAdapter()->CreateRHIBuffer(
+				DispatchRaysBufferDesc,
+				256,
+				CreateDesc,
+				ED3D12ResourceStateMode::MultiState,
+				ED3D12Access::IndirectArgs);
+		}
 	}
 #endif // D3D12_RHI_RAYTRACING
 
diff --git a/Engine/Source/Runtime/D3D12RHI/Private/D3D12RHI.cpp b/Engine/Source/Runtime/D3D12RHI/Private/D3D12RHI.cpp
--- a/Engine/Source/Runtime/D3D12RHI/Private/D3D12RHI.cpp
+++ b/Engine/Source/Runtime/D3D12RHI/Private/D3D12RHI.cpp
@@ -534,16 +534,20 @@ void FD3D12DynamicRHI::RHIEndFrame_RenderThread(FRHICommandListImmediate& RHICmd
 	RHICmdList.EnqueueLambdaMultiPipe(GetEnabledRHIPipelines(), FRHICommandListBase::EThreadFence::Enabled, TEXT("D3D12 BeginFrame"),
 		[this](FD3D12ContextArray const& Contexts)
 	{
-		for (auto& Adapter : ChosenAdapters)
+		for (const TSharedPtr<FD3D12Adapter>& Adapter : ChosenAdapters)
 		{
-			for (auto& Device : Adapter->GetDevices())
+			for (FD3D12Device* const Device : Adapter->GetDevices())
 			{
 				Device->GetDefaultBufferAllocator().BeginFrame(Contexts);
 				Device->GetTextureAllocator().BeginFrame(Contexts);
 
 			#if D3D12_RHI_RAYTRACING
 				// @todo dev-pr - explicit use of graphics context - nothing is synchronizing async compute - needs refactor
-				Device->GetRayTracingCompactionRequestHandler()->Update(*Contexts[ERHIPipeline::Graphics]->GetSingleDeviceContext(Device->GetGPUIndex()));
+				FD3D12RayTracingCompactionRequestHandler* const RayTracingCompactionRequestHandler = Device->GetRayTracingCompactionRequestHandler();
+				if (RayTracingCompactionRequestHandler)
+				{
+					RayTracingCompactionRequestHandler->Update(*Contexts[ERHIPipeline::Graphics]->GetSingleDeviceContext(Device->GetGPUIndex()));
+				}
 			#endif
 			}
 		}
diff --git a/Engine/Source/Runtime/D3D12RHI/Private/D3D12RayTracing.cpp b/Engine/Source/Runtime/D3D12RHI/Private/D3D12RayTracing.cpp
--- a/Engine/Source/Runtime/D3D12RHI/Private/D3D12RayTracing.cpp
+++ b/Engine/Source/Runtime/D3D12RHI/Private/D3D12RayTracing.cpp
@@ -3165,21 +3165,27 @@ private:
 
 void FD3D12Device::InitRayTracing()
 {
-	LLM_SCOPE_BYNAME(TEXT("FD3D12RT"));
-	check(RayTracingPipelineCache == nullptr);
-	RayTracingPipelineCache = new FD3D12RayTracingPipelineCache(GetParentAdapter());
+	if (GMaxRHIFeatureLevel >= ERHIFeatureLevel::SM6)
+	{
+		LLM_SCOPE_BYNAME(TEXT("FD3D12RT"));
+		check(RayTracingPipelineCache == nullptr);
+		RayTracingPipelineCache = new FD3D12RayTracingPipelineCache(GetParentAdapter());
+	}
 }
 
 void FD3D12Device::CleanupRayTracing()
 {
-	delete RayTracingPipelineCache;
-	RayTracingPipelineCache = nullptr;
-
-	for (FD3D12Queue& Queue : Queues)
+	if (GMaxRHIFeatureLevel >= ERHIFeatureLevel::SM6)
 	{
-		delete Queue.RayTracingDispatchRaysDescBuffer;
-		Queue.RayTracingDispatchRaysDescBuffer = nullptr;
-	}	
+		delete RayTracingPipelineCache;
+		RayTracingPipelineCache = nullptr;
+
+		for (FD3D12Queue& Queue : Queues)
+		{
+			delete Queue.RayTracingDispatchRaysDescBuffer;
+			Queue.RayTracingDispatchRaysDescBuffer = nullptr;
+		}
+	}
 }
 
 static D3D12_RAYTRACING_INSTANCE_FLAGS TranslateRayTracingInstanceFlags(ERayTracingInstanceFlags InFlags)
diff --git a/Engine/Source/Runtime/D3D12RHI/Private/Windows/WindowsD3D12Device.cpp b/Engine/Source/Runtime/D3D12RHI/Private/Windows/WindowsD3D12Device.cpp
--- a/Engine/Source/Runtime/D3D12RHI/Private/Windows/WindowsD3D12Device.cpp
+++ b/Engine/Source/Runtime/D3D12RHI/Private/Windows/WindowsD3D12Device.cpp
@@ -1574,54 +1574,64 @@ void FD3D12DynamicRHI::Init()
 				UE_LOG(LogD3D12RHI, Warning, TEXT("Failed to query support for 64 bit atomics"));
 			}
 
-			NvAPI_Status NvStatusSER = NvAPI_D3D12_IsNvShaderExtnOpCodeSupported(GetAdapter().GetD3DDevice(), NV_EXTN_OP_HIT_OBJECT_REORDER_THREAD, &GRHIGlobals.SupportsShaderExecutionReordering);
-			if (NvStatusSER == NVAPI_OK && GRHIGlobals.SupportsShaderExecutionReordering)
+			// NvAPI_D3D12_GetRaytracingCaps appears to leak about 100 MB so avoid calling if RayTracing is impossible
+			if (GMaxRHIFeatureLevel >= ERHIFeatureLevel::SM6)
 			{
-				NVAPI_D3D12_RAYTRACING_THREAD_REORDERING_CAPS ReorderCaps = NVAPI_D3D12_RAYTRACING_THREAD_REORDERING_CAP_NONE;
-				NvStatusSER = NvAPI_D3D12_GetRaytracingCaps(
-					GetAdapter().GetD3DDevice(),
-					NVAPI_D3D12_RAYTRACING_CAPS_TYPE_THREAD_REORDERING,
-					&ReorderCaps,
-					sizeof(ReorderCaps));
-
-				if (NvStatusSER != NVAPI_OK || ReorderCaps == NVAPI_D3D12_RAYTRACING_THREAD_REORDERING_CAP_NONE)
+				NvAPI_Status NvStatusSER = NvAPI_D3D12_IsNvShaderExtnOpCodeSupported(GetAdapter().GetD3DDevice(), NV_EXTN_OP_HIT_OBJECT_REORDER_THREAD, &GRHIGlobals.SupportsShaderExecutionReordering);
+				if (NvStatusSER == NVAPI_OK && GRHIGlobals.SupportsShaderExecutionReordering)
 				{
-					GRHIGlobals.SupportsShaderExecutionReordering = false;
+					NVAPI_D3D12_RAYTRACING_THREAD_REORDERING_CAPS ReorderCaps = NVAPI_D3D12_RAYTRACING_THREAD_REORDERING_CAP_NONE;
+					NvStatusSER = NvAPI_D3D12_GetRaytracingCaps(
+						GetAdapter().GetD3DDevice(),
+						NVAPI_D3D12_RAYTRACING_CAPS_TYPE_THREAD_REORDERING,
+						&ReorderCaps,
+						sizeof(ReorderCaps));
+
+					if (NvStatusSER != NVAPI_OK || ReorderCaps == NVAPI_D3D12_RAYTRACING_THREAD_REORDERING_CAP_NONE)
+					{
+						GRHIGlobals.SupportsShaderExecutionReordering = false;
+					}
 				}
-			}
 
-			if (GRHIGlobals.SupportsShaderExecutionReordering)
-			{
-				UE_LOG(LogD3D12RHI, Log, TEXT("NVIDIA Shader Execution Reordering interface supported!"));
-			}
-			else
-			{
-				UE_LOG(LogD3D12RHI, Log, TEXT("NVIDIA Shader Execution Reordering NOT supported!"));
-			}
-
-			// Ray Tracing Cluster Ops
-			{
-				GRHIGlobals.RayTracing.SupportsClusterOps = false;
-
-				NVAPI_D3D12_RAYTRACING_CLUSTER_OPERATIONS_CAPS ClusterOpsCaps = {};
-				NvAPI_Status NvStatusRayTracingClusterOps = NvAPI_D3D12_GetRaytracingCaps(GetAdapter().GetD3DDevice5(), NVAPI_D3D12_RAYTRACING_CAPS_TYPE_CLUSTER_OPERATIONS, &ClusterOpsCaps, sizeof(NVAPI_D3D12_RAYTRACING_CLUSTER_OPERATIONS_CAPS));
-
-				if (NvStatusRayTracingClusterOps == NVAPI_OK && ClusterOpsCaps == NVAPI_D3D12_RAYTRACING_CLUSTER_OPERATIONS_CAP_STANDARD)
+				if (GRHIGlobals.SupportsShaderExecutionReordering)
 				{
-					GRHIGlobals.RayTracing.SupportsClusterOps = CheckPipelineStateCap(NVAPI_D3D12_PIPELINE_CREATION_STATE_FLAGS_ENABLE_CLUSTER_SUPPORT);
+					UE_LOG(LogD3D12RHI, Log, TEXT("NVIDIA Shader Execution Reordering interface supported!"));
 				}
-
-				if (GRHIGlobals.RayTracing.SupportsClusterOps)
+				else
 				{
-					GRHIGlobals.RayTracing.ClusterAccelerationStructureAlignment = NVAPI_D3D12_RAYTRACING_CLAS_BYTE_ALIGNMENT;
-					GRHIGlobals.RayTracing.ClusterAccelerationStructureTemplateAlignment = NVAPI_D3D12_RAYTRACING_CLUSTER_TEMPLATE_BYTE_ALIGNMENT;
-					UE_LOG(LogD3D12RHI, Log, TEXT("NVIDIA Ray Tracing Cluster Operations supported and enabled"));
+					UE_LOG(LogD3D12RHI, Log, TEXT("NVIDIA Shader Execution Reordering NOT supported!"));
 				}
-				else
+
+				// Ray Tracing Cluster Ops
 				{
-					UE_LOG(LogD3D12RHI, Log, TEXT("NVIDIA Ray Tracing Cluster Operations not supported on this device"));
+					GRHIGlobals.RayTracing.SupportsClusterOps = false;
+
+					NVAPI_D3D12_RAYTRACING_CLUSTER_OPERATIONS_CAPS ClusterOpsCaps = {};
+					NvAPI_Status NvStatusRayTracingClusterOps = NvAPI_D3D12_GetRaytracingCaps(GetAdapter().GetD3DDevice5(), NVAPI_D3D12_RAYTRACING_CAPS_TYPE_CLUSTER_OPERATIONS, &ClusterOpsCaps, sizeof(NVAPI_D3D12_RAYTRACING_CLUSTER_OPERATIONS_CAPS));
+
+					if (NvStatusRayTracingClusterOps == NVAPI_OK && ClusterOpsCaps == NVAPI_D3D12_RAYTRACING_CLUSTER_OPERATIONS_CAP_STANDARD)
+					{
+						GRHIGlobals.RayTracing.SupportsClusterOps = CheckPipelineStateCap(NVAPI_D3D12_PIPELINE_CREATION_STATE_FLAGS_ENABLE_CLUSTER_SUPPORT);
+					}
+
+					if (GRHIGlobals.RayTracing.SupportsClusterOps)
+					{
+						GRHIGlobals.RayTracing.ClusterAccelerationStructureAlignment = NVAPI_D3D12_RAYTRACING_CLAS_BYTE_ALIGNMENT;
+						GRHIGlobals.RayTracing.ClusterAccelerationStructureTemplateAlignment = NVAPI_D3D12_RAYTRACING_CLUSTER_TEMPLATE_BYTE_ALIGNMENT;
+						UE_LOG(LogD3D12RHI, Log, TEXT("NVIDIA Ray Tracing Cluster Operations supported and enabled"));
+					}
+					else
+					{
+						UE_LOG(LogD3D12RHI, Log, TEXT("NVIDIA Ray Tracing Cluster Operations not supported on this device"));
+					}
 				}
 			}
+			else
+			{
+				GRHIGlobals.SupportsShaderExecutionReordering = false;
+				GRHIGlobals.RayTracing.SupportsClusterOps = false;
+				UE_LOG(LogD3D12RHI, Log, TEXT("Skipped NVAPI RT queries since we the feature level is below SM6"));
+			}
 		}
 		else
 		{
