RHIFinishExternalComputeWork wasn't considering bindless when calling SetDescriptorHeaps.

#Synced-CL 49626131

diff --git a/Engine/Source/Runtime/D3D12RHI/Private/D3D12RHI.cpp b/Engine/Source/Runtime/D3D12RHI/Private/D3D12RHI.cpp
--- a/Engine/Source/Runtime/D3D12RHI/Private/D3D12RHI.cpp
+++ b/Engine/Source/Runtime/D3D12RHI/Private/D3D12RHI.cpp
@@ -740,7 +740,17 @@ void FD3D12DynamicRHI::RHIFinishExternalComputeWork(FRHICommandListBase& Executi
 	check(InCommandList == Context.GraphicsCommandList().GetNoRefCount());
 
 	Context.StateCache.ForceSetComputeRootSignature();
-	Context.StateCache.GetDescriptorCache()->SetDescriptorHeaps(ED3D12SetDescriptorHeapsFlags::ForceChanged);
+
+	ED3D12SetDescriptorHeapsFlags SetHeapFlags = ED3D12SetDescriptorHeapsFlags::ForceChanged;
+#if PLATFORM_SUPPORTS_BINDLESS_RENDERING
+	// Make sure we restore the correct heaps
+	if (Context.StateCache.GetDescriptorCache()->IsUsingBindlessHeap())
+	{
+		EnumAddFlags(SetHeapFlags, ED3D12SetDescriptorHeapsFlags::Bindless);
+	}
+#endif
+
+	Context.StateCache.GetDescriptorCache()->SetDescriptorHeaps(SetHeapFlags);
 }
 
 void FD3D12DynamicRHI::RHITransitionResource(FRHICommandList& RHICmdList, FRHITexture* InTexture, D3D12_RESOURCE_STATES InState, uint32 InSubResource)
