Remove RHIGetViewportBackBufferUAV
- No code paths use this anymore. Instead, UAVs of the back buffer texture are created via the standard RDG/RHI UAV functions.

#Synced-CL 49648710

diff --git a/Engine/Source/Runtime/D3D12RHI/Private/D3D12RHIPrivate.h b/Engine/Source/Runtime/D3D12RHI/Private/D3D12RHIPrivate.h
--- a/Engine/Source/Runtime/D3D12RHI/Private/D3D12RHIPrivate.h
+++ b/Engine/Source/Runtime/D3D12RHI/Private/D3D12RHIPrivate.h
@@ -395,9 +395,6 @@ public:
 	virtual bool RHIGetRenderQueryResult(FRHIRenderQuery* RenderQuery, uint64& OutResult, bool bWait, uint32 GPUIndex = INDEX_NONE) final override;
 	virtual uint32 RHIGetViewportNextPresentGPUIndex(FRHIViewport* Viewport) final override;
 	virtual FTextureRHIRef RHIGetViewportBackBuffer(FRHIViewport* Viewport) final override;
-#if D3D12RHI_SUPPORTS_UAV_BACKBUFFER
-	virtual FUnorderedAccessViewRHIRef RHIGetViewportBackBufferUAV(FRHIViewport* Viewport) final override;
-#endif
 #if D3D12RHI_SUPPORTS_ALIASED_RHI_RESOURCES
 	virtual void RHIAliasTextureResources(FTextureRHIRef& DestTexture, FTextureRHIRef& SrcTexture) final override;
 	virtual FTextureRHIRef RHICreateAliasedTexture(FTextureRHIRef& SourceTexture) final override;
diff --git a/Engine/Source/Runtime/D3D12RHI/Private/D3D12Viewport.cpp b/Engine/Source/Runtime/D3D12RHI/Private/D3D12Viewport.cpp
--- a/Engine/Source/Runtime/D3D12RHI/Private/D3D12Viewport.cpp
+++ b/Engine/Source/Runtime/D3D12RHI/Private/D3D12Viewport.cpp
@@ -276,25 +276,6 @@ void FD3D12Viewport::InitializeBackBufferArrays()
 	}
 }
 
-template<typename TResource, typename TSubresource = TResource>
-bool ReleaseBackBufferResource(TRefCountPtr<TResource>& Resource, const TCHAR* ErrorName, uint32 Index)
-{
-	const bool bValidReference = IsValidRef(Resource);
-	if (bValidReference)
-	{
-		// Tell the back buffer to delete immediately so that we can call resize.
-		for (TSubresource& SubResource : *Resource)
-		{
-			SubResource.GetResource()->DoNotDeferDelete();
-		}
-	}
-
-	Resource.SafeRelease();
-	check(Resource == nullptr);
-
-	return bValidReference;
-}
-
 void FD3D12Viewport::Resize(uint32 InSizeX, uint32 InSizeY, bool bInIsFullscreen, EPixelFormat PreferredPixelFormat)
 {
 	FD3D12Adapter* Adapter = GetParentAdapter();
@@ -334,29 +315,29 @@ void FD3D12Viewport::Resize(uint32 InSizeX, uint32 InSizeY, bool bInIsFullscreen
 		CustomPresent->OnBackBufferResize();
 	}
 
-#if D3D12RHI_SUPPORTS_UAV_BACKBUFFER
-	bool bWaitForBackBuffersUAVDelete = false;
-
-	// Release our backbuffer reference, as required by DXGI before calling ResizeBuffers.
-	for (uint32 Index = 0; Index < NumBackBuffers; Index++)
 	{
-		bWaitForBackBuffersUAVDelete |= ReleaseBackBufferResource<FD3D12UnorderedAccessView_RHI, FD3D12UnorderedAccessView>(BackBuffers[Index].UAV, TEXT("BackBuffer UAV"), Index);
-	}
+		auto ReleaseBackBufferTexture = [](TRefCountPtr<FD3D12Texture>&Texture)
+		{
+			if (IsValidRef(Texture))
+			{
+				// Tell the back buffer to delete immediately so that we can call resize.
+				for (FD3D12Texture& SubResource : *Texture)
+				{
+					SubResource.GetResource()->DoNotDeferDelete();
+				}
+			}
 
-	if (bWaitForBackBuffersUAVDelete)
-	{
-		// The D3D12 UAV releases don't happen immediately, but are pushed to a delete queue processed on the RHI Thread. We need to ensure these are processed before releasing the swapchain buffers
-        // Calling FlushRenderingCommands is enough because it calls ImmediateFlush(EImmediateFlushType::FlushRHIThreadFlushResources) / ImmediateFlush(EImmediateFlushType::FlushRHIThread) internally
-		FlushRenderingCommands();
-	}
-#endif
+			Texture.SafeRelease();
+			check(Texture == nullptr);
+		};
 
-	for (uint32 Index = 0; Index < NumBackBuffers; Index++)
-	{
-		ReleaseBackBufferResource(BackBuffers[Index].Texture,    TEXT("BackBuffer"),     Index);
+		for (uint32 Index = 0; Index < NumBackBuffers; Index++)
+		{
+			ReleaseBackBufferTexture(BackBuffers[Index].Texture);
 #if D3D12RHI_USE_SDR_BACKBUFFER
-		ReleaseBackBufferResource(BackBuffers[Index].TextureSDR, TEXT("SDR BackBuffer"), Index);
+			ReleaseBackBufferTexture(BackBuffers[Index].TextureSDR);
 #endif
+		}
 	}
 
 	ClearPresentQueue();
diff --git a/Engine/Source/Runtime/D3D12RHI/Private/D3D12Viewport.h b/Engine/Source/Runtime/D3D12RHI/Private/D3D12Viewport.h
--- a/Engine/Source/Runtime/D3D12RHI/Private/D3D12Viewport.h
+++ b/Engine/Source/Runtime/D3D12RHI/Private/D3D12Viewport.h
@@ -84,17 +84,8 @@ public:
 #endif
 	}
 
-#if D3D12RHI_SUPPORTS_UAV_BACKBUFFER
-	FD3D12UnorderedAccessView_RHI* GetBackBufferUAV_RenderThread() const
-	{
-		checkSlow(CurrentBackBuffer_RenderThread);
-		return CurrentBackBuffer_RenderThread->UAV;
-	}
-#endif
-
 	FD3D12Texture* GetBackBuffer_RHIThread(bool bWaitForLastPresent = true) const;
 
-
 	FD3D12Texture* GetSDRBackBuffer_RHIThread(bool bWaitForLastPresent = true) const
 	{
 		if (bWaitForLastPresent)
@@ -255,16 +246,13 @@ private:
 
 	struct FBackBufferData
 	{
-		TRefCountPtr<FD3D12Texture>                 Texture;
+		TRefCountPtr<FD3D12Texture> Texture;
 #if D3D12RHI_USE_SDR_BACKBUFFER
 		// When HDR is enabled, SDR backbuffers may be required on some architectures for game DVR or broadcasting
-		TRefCountPtr<FD3D12Texture>                 TextureSDR;
-#endif
-#if D3D12RHI_SUPPORTS_UAV_BACKBUFFER
-		TRefCountPtr<FD3D12UnorderedAccessView_RHI> UAV;
+		TRefCountPtr<FD3D12Texture> TextureSDR;
 #endif
 #if WITH_MGPU
-		uint32                                      GPUIndex = 0;
+		uint32                      GPUIndex = 0;
 #endif
 	};
 
diff --git a/Engine/Source/Runtime/RHI/Public/DynamicRHI.h b/Engine/Source/Runtime/RHI/Public/DynamicRHI.h
--- a/Engine/Source/Runtime/RHI/Public/DynamicRHI.h
+++ b/Engine/Source/Runtime/RHI/Public/DynamicRHI.h
@@ -646,11 +646,6 @@ public:
 	// FlushType: Thread safe
 	virtual FTextureRHIRef RHIGetViewportBackBuffer(FRHIViewport* Viewport) = 0;
 
-	virtual FUnorderedAccessViewRHIRef RHIGetViewportBackBufferUAV(FRHIViewport* ViewportRHI)
-	{
-		return FUnorderedAccessViewRHIRef();
-	}
-
 	virtual uint32 RHIGetHTilePlatformConfig(uint32 DepthWidth, uint32 DepthHeight) const
 	{
 		return 0;
@@ -1255,9 +1250,10 @@ inline FTextureRHIRef RHIGetViewportBackBuffer(FRHIViewport* Viewport)
 	return GDynamicRHI->RHIGetViewportBackBuffer(Viewport);
 }
 
+UE_DEPRECATED(5.8, "RHIGetViewportBackBufferUAV has been deprecated and always returns null. Use the standard UAV creation functions instead. The platform must support back buffer UAVs, which is indicated by the TexCreate_UAV flag in the back buffer texture descriptor.")
 inline FUnorderedAccessViewRHIRef RHIGetViewportBackBufferUAV(FRHIViewport* Viewport)
 {
-	return GDynamicRHI->RHIGetViewportBackBufferUAV(Viewport);
+	return nullptr;
 }
 
 UE_DEPRECATED(5.7, "RHIGetHTilePlatformConfig(uint32,uint32) has been deprecated. Please use RHIGetHTilePlatformConfig(const FRHITextureDesc&) instead.")
diff --git a/Engine/Source/Runtime/RHI/Public/RHIValidation.h b/Engine/Source/Runtime/RHI/Public/RHIValidation.h
--- a/Engine/Source/Runtime/RHI/Public/RHIValidation.h
+++ b/Engine/Source/Runtime/RHI/Public/RHIValidation.h
@@ -695,11 +695,6 @@ public:
 		return Texture;
 	}
 
-	virtual FUnorderedAccessViewRHIRef RHIGetViewportBackBufferUAV(FRHIViewport* ViewportRHI) override final
-	{
-		return RHI->RHIGetViewportBackBufferUAV(ViewportRHI);
-	}
-
 	virtual uint32 RHIGetHTilePlatformConfig(uint32 DepthWidth, uint32 DepthHeight) const override final
 	{
 		return RHI->RHIGetHTilePlatformConfig(DepthWidth, DepthHeight);
