Removing EnqueueLambda [&] captures.

#Synced-CL 49664761

diff --git a/Engine/Source/Runtime/D3D12RHI/Private/D3D12RenderTarget.cpp b/Engine/Source/Runtime/D3D12RHI/Private/D3D12RenderTarget.cpp
--- a/Engine/Source/Runtime/D3D12RHI/Private/D3D12RenderTarget.cpp
+++ b/Engine/Source/Runtime/D3D12RHI/Private/D3D12RenderTarget.cpp
@@ -436,12 +436,7 @@ TRefCountPtr<FD3D12Resource> FD3D12DynamicRHI::GetStagingTexture(FRHITexture* Te
 			Subresource = CalcSubresource(InFlags.GetMip(), bIsTextureArray ? InFlags.GetArrayIndex() : 0, TextureRHI->GetNumMips());
 		}
 
-		D3D12_BOX* RectPtr = nullptr; // API prefers NULL for entire texture.
-		if (Rect.left != 0 || Rect.top != 0 || Rect.right != SourceDesc.Width || Rect.bottom != SourceDesc.Height)
-		{
-			// ..Sub rectangle required, use the D3D12_BOX.
-			RectPtr = &Rect;
-		}
+		const bool bSubrect = (Rect.left != 0 || Rect.top != 0 || Rect.right != SourceDesc.Width || Rect.bottom != SourceDesc.Height);
 
 		D3D12_PLACED_SUBRESOURCE_FOOTPRINT DestFootprint{};
 		DestFootprint.Footprint.Depth = 1;
@@ -458,9 +453,10 @@ TRefCountPtr<FD3D12Resource> FD3D12DynamicRHI::GetStagingTexture(FRHITexture* Te
 		ReadbackHeapDesc = DestFootprint;
 
 		// TODO: transition only subresource not the whole resource
-		RHICmdList.TransitionInternal(FRHITransitionInfo(TextureRHI, ERHIAccess::Unknown, ERHIAccess::CopySrc, EResourceTransitionFlags::IgnoreAfterState));
+		RHICmdList.Transition(FRHITransitionInfo(TextureRHI, ERHIAccess::Unknown, ERHIAccess::CopySrc, EResourceTransitionFlags::IgnoreAfterState));
 
-		RHICmdList.EnqueueLambda([&](FRHICommandList& ExecutingCmdList)
+		RHICmdList.EnqueueLambda(
+		[Texture, GPUIndex, Rect, bSubrect, DestCopyLocation = MoveTemp(DestCopyLocation), SourceCopyLocation = MoveTemp(SourceCopyLocation)](FRHICommandList& ExecutingCmdList)
 		{
 			FD3D12CommandContext& Context = FD3D12CommandContext::Get(ExecutingCmdList, GPUIndex);
 
@@ -468,6 +464,13 @@ TRefCountPtr<FD3D12Resource> FD3D12DynamicRHI::GetStagingTexture(FRHITexture* Te
 			Context.FlushResourceBarriers();
 			// Upload heap doesn't need to transition
 
+			const D3D12_BOX* RectPtr = nullptr; // API prefers NULL for entire texture.
+			if (bSubrect)
+			{
+				// ..Sub rectangle required, use the D3D12_BOX.
+				RectPtr = &Rect;
+			}
+
 			Context.CopyTextureRegionChecked(
 				&DestCopyLocation,
 				0, 0, 0,
@@ -776,19 +779,20 @@ void FD3D12DynamicRHI::ReadSurfaceDataMSAARaw(FRHITexture* TextureRHI, FIntRect
 	// Can be optimized by doing all subsamples into a large enough rendertarget in one pass (multiple draw calls)
 	for (uint32 SampleIndex = 0; SampleIndex < NumSamples; ++SampleIndex)
 	{
-		RHICmdList.EnqueueLambda([&](FRHICommandListBase& ExecutingCmdList)
+		RHICmdList.EnqueueLambda(
+		[Texture, GPUIndex, SampleIndex, InRect, Rect, NonMSAADesc, NonMSAATexture2D, NonMSAARTV, StagingTexture2D, DestCopyLocation = MoveTemp(DestCopyLocation), SourceCopyLocation = MoveTemp(SourceCopyLocation)](FRHICommandListBase& ExecutingCmdList) mutable
 		{
 			FD3D12CommandContext& Context = FD3D12CommandContext::Get(ExecutingCmdList, GPUIndex);
 
 			// Resolve the sample to the non-MSAA render target.
 			Context.ResolveTextureUsingShader<FResolveSingleSamplePS>(
-				ResourceCast(TextureRHI->GetTexture2D()),
+				ResourceCast(Texture->GetTexture2D()),
 				NULL,
 				&NonMSAARTV,
 				NULL,
 				NonMSAADesc,
 				FResolveRect(InRect.Min.X, InRect.Min.Y, InRect.Max.X, InRect.Max.Y),
-				FResolveRect(0, 0, SizeX, SizeY),
+				FResolveRect(0, 0, InRect.Width(), InRect.Height()),
 				SampleIndex
 				);
 
@@ -971,7 +975,8 @@ void FD3D12DynamicRHI::RHIReadSurfaceFloatData(FRHITexture* TextureRHI, FIntRect
 	CD3DX12_TEXTURE_COPY_LOCATION SourceCopyLocation(Texture->GetResource()->GetResource(), Subresource);
 
 	FRHICommandListImmediate& RHICmdList = FRHICommandListImmediate::Get();
-	RHICmdList.EnqueueLambda([&](FRHICommandListBase& ExecutingCmdList)
+	RHICmdList.EnqueueLambda(
+	[Texture, GPUIndex, Rect, DestCopyLocation = MoveTemp(DestCopyLocation), SourceCopyLocation = MoveTemp(SourceCopyLocation)](FRHICommandListBase& ExecutingCmdList)
 	{
 		FD3D12CommandContext& Context = FD3D12CommandContext::Get(ExecutingCmdList, GPUIndex);
 
@@ -1090,7 +1095,8 @@ void FD3D12DynamicRHI::RHIRead3DSurfaceFloatData(FRHITexture* TextureRHI, FIntRe
 	CD3DX12_TEXTURE_COPY_LOCATION SourceCopyLocation(Texture->GetResource()->GetResource(), Subresource);
 
 	FRHICommandListImmediate& RHICmdList = FRHICommandListImmediate::Get();
-	RHICmdList.EnqueueLambda([&](FRHICommandListBase& ExecutingCmdList)
+	RHICmdList.EnqueueLambda(
+	[Texture, GPUIndex, Rect, DestCopyLocation = MoveTemp(DestCopyLocation), SourceCopyLocation = MoveTemp(SourceCopyLocation)](FRHICommandListBase& ExecutingCmdList)
 	{
 		FD3D12CommandContext& Context = FD3D12CommandContext::Get(ExecutingCmdList, GPUIndex);
 
diff --git a/Engine/Source/Runtime/OpenGLDrv/Private/OpenGLBuffer.cpp b/Engine/Source/Runtime/OpenGLDrv/Private/OpenGLBuffer.cpp
--- a/Engine/Source/Runtime/OpenGLDrv/Private/OpenGLBuffer.cpp
+++ b/Engine/Source/Runtime/OpenGLDrv/Private/OpenGLBuffer.cpp
@@ -316,7 +316,7 @@ void FOpenGLStagingBuffer::Initialize()
 	ShadowSize = 0;
 	Mapping = nullptr;
 	FRHICommandListImmediate& RHICmdList = FRHICommandListImmediate::Get();
-	RHICmdList.EnqueueLambda([&](FRHICommandListImmediate&)
+	RHICmdList.EnqueueLambda([this](FRHICommandListImmediate&)
 	{
 		VERIFY_GL_SCOPE();
 		glGenBuffers(1, &ShadowBuffer);
@@ -339,7 +339,7 @@ void* FOpenGLStagingBuffer::Lock(uint32 Offset, uint32 NumBytes)
 	{
 		void* ReturnValue = nullptr;
 		FRHICommandListImmediate& RHICmdList = FRHICommandListImmediate::Get();
-		RHICmdList.EnqueueLambda([&](FRHICommandListImmediate&) {
+		RHICmdList.EnqueueLambda([this, Offset, NumBytes, &ReturnValue](FRHICommandListImmediate&) {
 			VERIFY_GL_SCOPE();
 			check(ShadowBuffer != 0);
 			glBindBuffer(GL_COPY_WRITE_BUFFER, ShadowBuffer);
@@ -365,7 +365,7 @@ void FOpenGLStagingBuffer::Unlock()
 	if (!FOpenGL::SupportsBufferStorage() || !OpenGLConsoleVariables::bUsePersistentMappingStagingBuffer)
 	{
 		FRHICommandListImmediate& RHICmdList = FRHICommandListImmediate::Get();
-		RHICmdList.EnqueueLambda([&](FRHICommandListImmediate&) {
+		RHICmdList.EnqueueLambda([this](FRHICommandListImmediate&) {
 			FOpenGL::UnmapBuffer(GL_COPY_WRITE_BUFFER);
 			Mapping = nullptr;
 			glBindBuffer(GL_COPY_WRITE_BUFFER, 0);
diff --git a/Engine/Source/Runtime/OpenGLDrv/Private/OpenGLCommands.cpp b/Engine/Source/Runtime/OpenGLDrv/Private/OpenGLCommands.cpp
--- a/Engine/Source/Runtime/OpenGLDrv/Private/OpenGLCommands.cpp
+++ b/Engine/Source/Runtime/OpenGLDrv/Private/OpenGLCommands.cpp
@@ -2562,7 +2562,7 @@ void FOpenGLDynamicRHI::RHIClearMRT(const bool* bClearColorArray,int32 NumClearC
 void FOpenGLDynamicRHI::RHIBlockUntilGPUIdle()
 {
 	FRHICommandListImmediate& RHICmdList = FRHICommandListImmediate::Get();
-	RHICmdList.EnqueueLambda([&](FRHICommandListImmediate&)
+	RHICmdList.EnqueueLambda([](FRHICommandListImmediate&)
 	{
 		FOpenGL::Flush();
 	});
diff --git a/Engine/Source/Runtime/OpenGLDrv/Private/OpenGLRenderTarget.cpp b/Engine/Source/Runtime/OpenGLDrv/Private/OpenGLRenderTarget.cpp
--- a/Engine/Source/Runtime/OpenGLDrv/Private/OpenGLRenderTarget.cpp
+++ b/Engine/Source/Runtime/OpenGLDrv/Private/OpenGLRenderTarget.cpp
@@ -747,8 +747,9 @@ void FOpenGLDynamicRHI::RHIReadSurfaceData(FRHITexture* TextureRHI,FIntRect Rect
 
 	FRHICommandListImmediate& RHICmdList = FRHICommandListImmediate::Get();
 	
-	RHICmdList.EnqueueLambda([&](FRHICommandListImmediate&)
+	RHICmdList.EnqueueLambda([this, TextureRHI, Rect, InFlags, &OutData](FRHICommandListImmediate&)
 	{
+		const uint32 Size = Rect.Width() * Rect.Height();
 		TArray<uint8> Temp;
 
 		ReadSurfaceDataRaw(TextureRHI, Rect, Temp, InFlags);
@@ -770,7 +771,7 @@ void FOpenGLDynamicRHI::RHIReadSurfaceData(FRHITexture* TextureRHI, FIntRect Rec
 	
 	FRHICommandListImmediate& RHICmdList = FRHICommandListImmediate::Get();
 	
-	RHICmdList.EnqueueLambda([&](FRHICommandListImmediate&)
+	RHICmdList.EnqueueLambda([this, TextureRHI, Rect, InFlags, &OutData](FRHICommandListImmediate&)
 	{
 		VERIFY_GL_SCOPE();
 
@@ -812,7 +813,7 @@ void FOpenGLDynamicRHI::RHIMapStagingSurface_RenderThread(class FRHICommandListI
 	// Fence is not important, GL driver handles synchonization
 	RHICmdList.ImmediateFlush(EImmediateFlushType::FlushRHIThread);
 
-	RHICmdList.EnqueueLambda([&](FRHICommandListImmediate&)
+	RHICmdList.EnqueueLambda([TextureRHI, &OutData, &OutWidth, &OutHeight](FRHICommandListImmediate&)
 	{
 		VERIFY_GL_SCOPE();
 
@@ -857,7 +858,7 @@ void FOpenGLDynamicRHI::RHIReadSurfaceFloatData(FRHITexture* TextureRHI,FIntRect
 {
 	FRHICommandListImmediate& RHICmdList = FRHICommandListImmediate::Get();
 
-	RHICmdList.EnqueueLambda([&](FRHICommandListImmediate&)
+	RHICmdList.EnqueueLambda([this, TextureRHI, Rect, CubeFace, ArrayIndex, MipIndex, &OutData](FRHICommandListImmediate&)
 	{
 		VERIFY_GL_SCOPE();	
 
@@ -934,7 +935,7 @@ void FOpenGLDynamicRHI::RHIRead3DSurfaceFloatData(FRHITexture* TextureRHI,FIntRe
 {
 	FRHICommandListImmediate& RHICmdList = FRHICommandListImmediate::Get();
 
-	RHICmdList.EnqueueLambda([&](FRHICommandListImmediate&)
+	RHICmdList.EnqueueLambda([this, TextureRHI, Rect, ZMinMax, &OutData](FRHICommandListImmediate&)
 	{
 		VERIFY_GL_SCOPE();
 
diff --git a/Engine/Source/Runtime/OpenGLDrv/Private/OpenGLTexture.cpp b/Engine/Source/Runtime/OpenGLDrv/Private/OpenGLTexture.cpp
--- a/Engine/Source/Runtime/OpenGLDrv/Private/OpenGLTexture.cpp
+++ b/Engine/Source/Runtime/OpenGLDrv/Private/OpenGLTexture.cpp
@@ -308,7 +308,7 @@ FOpenGLTexture::FOpenGLTexture(FOpenGLTexture& Other, const FString& Name, EAlia
 	, bMultisampleRenderbuffer(Other.bMultisampleRenderbuffer)
 {
 	FRHICommandListImmediate& RHICmdList = FRHICommandListImmediate::Get();
-	RHICmdList.EnqueueLambda([&](FRHICommandListImmediate&)
+	RHICmdList.EnqueueLambda([this, &Other](FRHICommandListImmediate&)
 	{
 		AliasResources(Other);
 	});
