Disable PresentOnSubmissionThread causing hitches on D3D12 platforms

#Synced-CL 49720845

diff --git a/Engine/Source/Runtime/D3D12RHI/Private/D3D12Viewport.cpp b/Engine/Source/Runtime/D3D12RHI/Private/D3D12Viewport.cpp
--- a/Engine/Source/Runtime/D3D12RHI/Private/D3D12Viewport.cpp
+++ b/Engine/Source/Runtime/D3D12RHI/Private/D3D12Viewport.cpp
@@ -35,7 +35,7 @@ namespace D3D12RHI
 #endif
 	};
 
-	bool GUseSubmissionThreadPresent = true;
+	bool GUseSubmissionThreadPresent = false;
 	static FAutoConsoleVariableRef CVarD3D12SubmissionTimeout(
 		TEXT("r.D3D12.PresentOnSubmissionThread"),
 		GUseSubmissionThreadPresent,
