[HWRT] Support setting up SBT entries without providing a FRHIRayTracingGeometry at RHI level.
- For Nanite ray tracing we get vertex/index data directly from ClusterPageData buffer, and the VB/IB used to build the BLAS are transient so binding those is not safe.
- r.RayTracing.Nanite.Mode 1 still needs to provide a FRHIRayTracingGeometry since we need to access the segments FirstPrimitive in the shaders, but that won't be necessary when using CLAS APIs.

#Synced-CL 49716836

diff --git a/Engine/Source/Runtime/D3D12RHI/Private/D3D12RayTracing.cpp b/Engine/Source/Runtime/D3D12RHI/Private/D3D12RayTracing.cpp
--- a/Engine/Source/Runtime/D3D12RHI/Private/D3D12RayTracing.cpp
+++ b/Engine/Source/Runtime/D3D12RHI/Private/D3D12RayTracing.cpp
@@ -2151,9 +2151,19 @@ public:
 		return (RecordIndex % NumShaderSlotsPerGeometrySegment == 0) ? RecordIndex / NumShaderSlotsPerGeometrySegment : INDEX_NONE;
 	}
 
-	void SetHitGroupGeometrySystemParameters(uint32 WorkerIndex, ERayTracingLocalShaderBindingType BindingType, uint32 RecordIndex, const FD3D12RayTracingGeometry* Geometry, uint32 UserData, uint32 GeometrySegmentIndex)
+	void SetHitGroupGeometrySystemParameters(uint32 WorkerIndex, ERayTracingLocalShaderBindingType BindingType, uint32 RecordIndex, const FD3D12RayTracingGeometry* Geometry, uint32 GeometrySegmentIndex, uint32 UserData)
 	{	
-		FD3D12HitGroupSystemParameters SystemParameters = Geometry->HitGroupSystemParameters[Device->GetGPUIndex()][GeometrySegmentIndex];
+		FD3D12HitGroupSystemParameters SystemParameters;
+
+		if (Geometry != nullptr)
+		{
+			SystemParameters = Geometry->HitGroupSystemParameters[Device->GetGPUIndex()][GeometrySegmentIndex];
+		}
+		else
+		{
+			FMemory::Memset(&SystemParameters, 0, sizeof(SystemParameters));
+		}
+
 		SystemParameters.RootConstants.UserData = UserData;		
 
 		WorkerData[WorkerIndex].MaxUsedHitRecordIndex = FMath::Max(WorkerData[WorkerIndex].MaxUsedHitRecordIndex, (int32)RecordIndex);
@@ -2199,7 +2209,10 @@ public:
 				// If persistently stored then also listen to hit record updates for stream in/out events of the ray tracing geometry data
 				if (BindingType == ERayTracingLocalShaderBindingType::Persistent)
 				{
-					HitRecordData[RecordIndex].UpdateGeometryListener.Init(this, Geometry, RecordIndex, UserData, GeometrySegmentIndex);
+					if (Geometry != nullptr)
+					{
+						HitRecordData[RecordIndex].UpdateGeometryListener.Init(this, Geometry, RecordIndex, UserData, GeometrySegmentIndex);
+					}
 					RecordData.State = FRecordData::EState::Persistent;
 				}				
 				else // if (RecordData.State == FRecordData::EState::Uninitialized)
@@ -5826,6 +5839,7 @@ static void SetRayTracingHitGroup(
 	checkf(RecordIndex < ShaderTable->NumHitRecords, TEXT("Hit group record index is invalid. Make sure that NumGeometrySegments and NumShaderSlotsPerGeometrySegment is correct in FRayTracingShaderBindingTableInitializer."));
 
 #if DO_CHECK
+	if (Geometry != nullptr)
 	{
 		const uint32 NumGeometrySegments = Geometry->GetNumSegments();
 		checkf(GeometrySegmentIndex < NumGeometrySegments, TEXT("Segment %d is out of range for ray tracing geometry '%s' that contains %d segments"),
@@ -5847,7 +5861,7 @@ static void SetRayTracingHitGroup(
 	}
 #endif // DO_CHECK
 
-	ShaderTable->SetHitGroupGeometrySystemParameters(WorkerIndex, BindingType, RecordIndex, Geometry, UserData, GeometrySegmentIndex);
+	ShaderTable->SetHitGroupGeometrySystemParameters(WorkerIndex, BindingType, RecordIndex, Geometry, GeometrySegmentIndex, UserData);
 
 	if (EnumHasAnyFlags(ShaderTable->ShaderBindingMode, ERayTracingShaderBindingMode::RTPSO) && Pipeline)
 	{
