Simplify viewport presentation by merging RHIAdvanceFrameForGetViewportBackBuffer into RHIEndDrawingViewport

Added GRHIGlobals.SupportsConcurrentTranslateAndSubmit
- Allow RHI translation and submission tasks to overlap. Requires platform RHIs to avoid submitting work to the GPU from inside RHI translation tasks. Work may only be submitted from RHISubmitCommandLists.
- Not currently enabled on any platform.

#Synced-CL 49885160

diff --git a/Engine/Source/Runtime/Apple/MetalRHI/Private/MetalDynamicRHI.h b/Engine/Source/Runtime/Apple/MetalRHI/Private/MetalDynamicRHI.h
--- a/Engine/Source/Runtime/Apple/MetalRHI/Private/MetalDynamicRHI.h
+++ b/Engine/Source/Runtime/Apple/MetalRHI/Private/MetalDynamicRHI.h
@@ -176,11 +176,11 @@ public:
 	virtual bool RHIGetRenderQueryResult(FRHIRenderQuery* RenderQuery, uint64& OutResult, bool bWait, uint32 GPUIndex = INDEX_NONE) final override;
 	virtual FTextureRHIRef RHIGetViewportBackBuffer(FRHIViewport* Viewport) final override;
 	virtual FTextureRHIRef RHIGetViewportBackBufferAtLayerIndex(FRHIViewport* Viewport, uint32 DrawLayerIndex) final override;
-	virtual void RHIAdvanceFrameForGetViewportBackBuffer(FRHIViewport* Viewport, bool bPresent) final override;
 	virtual void RHIFlushResources() final override;
 	virtual FViewportRHIRef RHICreateViewport(void* WindowHandle, uint32 SizeX, uint32 SizeY, bool bIsFullscreen, EPixelFormat PreferredPixelFormat) final override;
 	virtual void RHIResizeViewport(FRHIViewport* Viewport, uint32 SizeX, uint32 SizeY, bool bIsFullscreen) final override;
 	virtual void RHIResizeViewport(FRHIViewport* Viewport, uint32 SizeX, uint32 SizeY, bool bIsFullscreen, EPixelFormat PreferredPixelFormat) final override;
+	virtual void RHIEndDrawingViewport(FRHICommandListImmediate& RHICmdList, FRHIViewport* Viewport, const FRHIPresentArgs& PresentArgs) final override;
 	virtual void RHITick(float DeltaTime) final override;
 	virtual void RHIBlockUntilGPUIdle() final override;
 	virtual bool RHIGetAvailableResolutions(FScreenResolutionArray& Resolutions, bool bIgnoreRefreshRate) final override;
diff --git a/Engine/Source/Runtime/Apple/MetalRHI/Private/MetalViewport.cpp b/Engine/Source/Runtime/Apple/MetalRHI/Private/MetalViewport.cpp
--- a/Engine/Source/Runtime/Apple/MetalRHI/Private/MetalViewport.cpp
+++ b/Engine/Source/Runtime/Apple/MetalRHI/Private/MetalViewport.cpp
@@ -1199,11 +1199,9 @@ void FMetalDynamicRHI::RHITick( float DeltaTime )
  *	Viewport functions.
  *=============================================================================*/
 
-void FMetalRHICommandContext::RHIEndDrawingViewport(FRHIViewport* ViewportRHI, const FRHIPresentArgs& InPresentArgs)
+void FMetalRHICommandContext::EndDrawingViewport(FMetalViewport* Viewport, const FRHIPresentArgs& InPresentArgs)
 {
     MTL_SCOPED_AUTORELEASE_POOL;
-    
-	FMetalViewport* Viewport = ResourceCast(ViewportRHI);
 
 	// enqueue a present if desired
 	static bool const bOffscreenOnly = FParse::Param(FCommandLine::Get(), TEXT("MetalOffscreenOnly"));
@@ -1255,6 +1253,30 @@ void FMetalRHICommandContext::RHIEndDrawingViewport(FRHIViewport* ViewportRHI, c
 	Viewport->ReleaseDrawable();
 }
 
+void FMetalDynamicRHI::RHIEndDrawingViewport(FRHICommandListImmediate& RHICmdList, FRHIViewport* ViewportRHI, const FRHIPresentArgs& PresentArgs)
+{
+	//
+	// Make sure all prior graphics and async compute work has been submitted. This is necessary because Metal RHI handles the present from inside the below lambda.
+	// Since the present is not yet pipelined via payloads etc, we must ensure prior work has been handed to RHISubmitCommandLists before we translate the lambda.
+	//
+	// In future, Present() itself should be an enqueued command, and platform RHIs should never implicitly send commands to the GPU queues during RHI translation.
+	// All GPU work should be submitted via RHISubmitCommandLists.
+	//
+	// This is also why there is a task dependency between RHI translate and RHI submission tasks in RHICommandList.cpp.
+	// See GRHIGlobals.SupportsConcurrentTranslateAndSubmit
+	//
+	RHICmdList.ImmediateFlush(EImmediateFlushType::DispatchToRHIThread);
+
+	{
+		FRHICommandListScopedPipeline PipelineScope(RHICmdList, ERHIPipeline::Graphics);
+	    RHICmdList.EnqueueLambda(TEXT("RHIEndDrawingViewport"), [ViewportRHI, PresentArgs](FRHICommandListImmediate& ExecutingCmdList)
+	    {
+			FMetalRHICommandContext& Context = FMetalRHICommandContext::Get(ExecutingCmdList);
+			Context.EndDrawingViewport(ResourceCast(ViewportRHI), PresentArgs);
+	    });
+	}
+}
+
 FTextureRHIRef FMetalDynamicRHI::RHIGetViewportBackBuffer(FRHIViewport* ViewportRHI)
 {
     MTL_SCOPED_AUTORELEASE_POOL;
@@ -1270,7 +1292,3 @@ FTextureRHIRef FMetalDynamicRHI::RHIGetViewportBackBufferAtLayerIndex(FRHIViewpo
 	FMetalViewport* Viewport = ResourceCast(ViewportRHI);
 	return FTextureRHIRef(Viewport->GetBackBuffer(EMetalViewportAccessRenderer, DrawLayerIndex).GetReference());
 }
-
-void FMetalDynamicRHI::RHIAdvanceFrameForGetViewportBackBuffer(FRHIViewport* ViewportRHI, bool bPresent)
-{
-}
diff --git a/Engine/Source/Runtime/Apple/MetalRHI/Public/MetalRHIContext.h b/Engine/Source/Runtime/Apple/MetalRHI/Public/MetalRHIContext.h
--- a/Engine/Source/Runtime/Apple/MetalRHI/Public/MetalRHIContext.h
+++ b/Engine/Source/Runtime/Apple/MetalRHI/Public/MetalRHIContext.h
@@ -80,9 +80,8 @@ public:
 	virtual void RHIBeginRenderQuery(FRHIRenderQuery* RenderQuery) final override;
 	virtual void RHIEndRenderQuery(FRHIRenderQuery* RenderQuery) final override;
 	
-	// This method is queued with an RHIThread, otherwise it will flush after it is queued; without an RHI thread there is no benefit to queuing this frame advance commands
-	virtual void RHIEndDrawingViewport(FRHIViewport* Viewport, const FRHIPresentArgs& InPresentArgs) override;
-	
+	void EndDrawingViewport(FMetalViewport* Viewport, const FRHIPresentArgs& InPresentArgs);
+
 	virtual void RHISetStreamSource(uint32 StreamIndex, FRHIBuffer* VertexBuffer, uint32 Offset) final override;
 
 	// @param MinX including like Win32 RECT
diff --git a/Engine/Source/Runtime/D3D12RHI/Private/D3D12CommandContext.h b/Engine/Source/Runtime/D3D12RHI/Private/D3D12CommandContext.h
--- a/Engine/Source/Runtime/D3D12RHI/Private/D3D12CommandContext.h
+++ b/Engine/Source/Runtime/D3D12RHI/Private/D3D12CommandContext.h
@@ -483,8 +483,6 @@ class FD3D12CommandContextBase : public IRHICommandContext, public FD3D12Adapter
 public:
 	FD3D12CommandContextBase(FD3D12Adapter* InParent, FRHIGPUMask InGPUMask);
 
-	void RHIEndDrawingViewport(FRHIViewport* Viewport, const FRHIPresentArgs& InPresentArgs) final override;
-
 	FRHIGPUMask GetGPUMask() const { return GPUMask; }
 	FRHIGPUMask GetPhysicalGPUMask() const { return PhysicalGPUMask; }
 
diff --git a/Engine/Source/Runtime/D3D12RHI/Private/D3D12RHIPrivate.h b/Engine/Source/Runtime/D3D12RHI/Private/D3D12RHIPrivate.h
--- a/Engine/Source/Runtime/D3D12RHI/Private/D3D12RHIPrivate.h
+++ b/Engine/Source/Runtime/D3D12RHI/Private/D3D12RHIPrivate.h
@@ -403,11 +403,11 @@ public:
 	virtual uint64 RHIComputeStatePrecachePSOHash(const FGraphicsPipelineStateInitializer& Initializer) final override;
 	virtual uint64 RHIComputePrecachePSOHash(const FGraphicsPipelineStateInitializer& Initializer) final override;
 	virtual bool RHIMatchPrecachePSOInitializers(const FGraphicsPipelineStateInitializer& LHS, const FGraphicsPipelineStateInitializer& RHS) final override;
-	virtual void RHIAdvanceFrameForGetViewportBackBuffer(FRHIViewport* Viewport, bool bPresent) final override;
 	virtual void RHIFlushResources() final override;
 	virtual FViewportRHIRef RHICreateViewport(void* WindowHandle, uint32 SizeX, uint32 SizeY, bool bIsFullscreen, EPixelFormat PreferredPixelFormat) final override;
 	virtual void RHIResizeViewport(FRHIViewport* Viewport, uint32 SizeX, uint32 SizeY, bool bIsFullscreen) final override;
 	virtual void RHIResizeViewport(FRHIViewport* ViewportRHI, uint32 SizeX, uint32 SizeY, bool bIsFullscreen, EPixelFormat PreferredPixelFormat) final override;
+	virtual void RHIEndDrawingViewport(FRHICommandListImmediate& RHICmdList, FRHIViewport* Viewport, const FRHIPresentArgs& PresentArgs) final override;
 #if PLATFORM_WINDOWS
 	virtual void RHIHandleDisplayChange() final override;
 #endif
diff --git a/Engine/Source/Runtime/D3D12RHI/Private/D3D12Viewport.cpp b/Engine/Source/Runtime/D3D12RHI/Private/D3D12Viewport.cpp
--- a/Engine/Source/Runtime/D3D12RHI/Private/D3D12Viewport.cpp
+++ b/Engine/Source/Runtime/D3D12RHI/Private/D3D12Viewport.cpp
@@ -883,72 +883,86 @@ void FD3D12DynamicRHI::RHITick(float DeltaTime)
  *	Viewport functions.
  *=============================================================================*/
 
-void FD3D12CommandContextBase::RHIEndDrawingViewport(FRHIViewport* ViewportRHI, const FRHIPresentArgs& InPresentArgs)
+void FD3D12DynamicRHI::RHIEndDrawingViewport(FRHICommandListImmediate& RHICmdList, FRHIViewport* ViewportRHI, const FRHIPresentArgs& PresentArgs)
 {
-	FD3D12Viewport* Viewport = FD3D12DynamicRHI::ResourceCast(ViewportRHI);
-	SCOPE_CYCLE_COUNTER(STAT_D3D12PresentTime);
+	//
+	// Make sure all prior graphics and async compute work has been submitted. This is necessary because D3D12 RHI handles the present from inside the below lambda.
+	// Since the present is not yet pipelined via payloads etc, we must ensure prior work has been handed to RHISubmitCommandLists before we translate the lambda.
+	//
+	// In future, Present() itself should be an enqueued command, and platform RHIs should never implicitly send commands to the GPU queues during RHI translation.
+	// All GPU work should be submitted via RHISubmitCommandLists.
+	//
+	// This is also why there is a task dependency between RHI translate and RHI submission tasks in RHICommandList.cpp.
+	// See GRHIGlobals.SupportsConcurrentTranslateAndSubmit
+	//
+	RHICmdList.ImmediateFlush(EImmediateFlushType::DispatchToRHIThread);
 
-#if !UE_BUILD_SHIPPING
-	if (RHIConsoleVariables::LogViewportEvents)
 	{
-		const FString& ThreadName = FThreadManager::GetThreadName(FPlatformTLS::GetCurrentThreadId());
-		UE_LOG(LogD3D12RHI, Log, TEXT("Thread %s: RHIEndDrawingViewport (Viewport %#016llx: BackBuffer %#016llx: CmdList: %016llx)"),
-			ThreadName.GetCharArray().GetData(),
-			Viewport,
-			Viewport->GetBackBuffer_RHIThread(),
-			GetSingleDeviceContext(0)->BaseCommandList().GetNoRefCount()
-		);
-	}
-#endif
+		FRHICommandListScopedPipeline PipelineScope(RHICmdList, ERHIPipeline::Graphics);
+		RHICmdList.EnqueueLambda(TEXT("RHIEndDrawingViewport"), [this, ViewportRHI, PresentArgs](FRHICommandListImmediate& ExecutingCmdList)
+		{
+			FD3D12CommandContextBase& Context = FD3D12CommandContextBase::Get(ExecutingCmdList);
 
-	const bool bNativelyPresented = Viewport->EndRendering(*this, InPresentArgs);
+			FD3D12Viewport* Viewport = FD3D12DynamicRHI::ResourceCast(ViewportRHI);
+			SCOPE_CYCLE_COUNTER(STAT_D3D12PresentTime);
 
-	if ((bNativelyPresented) || (GUseSubmissionThreadPresent))
-	{
-		FRenderThreadIdleScope Scope(ERenderThreadIdleTypes::WaitingForAllOtherSleep);
-		static const auto CFinishFrameVar = IConsoleManager::Get().FindTConsoleVariableDataInt(TEXT("r.FinishCurrentFrame"));
+		#if !UE_BUILD_SHIPPING
+			if (RHIConsoleVariables::LogViewportEvents)
+			{
+				const FString& ThreadName = FThreadManager::GetThreadName(FPlatformTLS::GetCurrentThreadId());
+				UE_LOG(LogD3D12RHI, Log, TEXT("Thread %s: RHIEndDrawingViewport (Viewport %#016llx: BackBuffer %#016llx: CmdList: %016llx)"),
+					ThreadName.GetCharArray().GetData(),
+					Viewport,
+					Viewport->GetBackBuffer_RHIThread(),
+					Context.GetSingleDeviceContext(0)->BaseCommandList().GetNoRefCount()
+				);
+			}
+		#endif
 
-		int32 FinishCurrentFrame = CFinishFrameVar->GetValueOnRenderThread();
+			const bool bNativelyPresented = Viewport->EndRendering(Context, PresentArgs);
 
-		if (FinishCurrentFrame == 0)
-		{
-			// Wait for the GPU to finish rendering the previous frame before finishing this frame.
-			Viewport->WaitForFrameEventCompletion();
-			Viewport->IssueFrameEvent();
-		}
-		else if (FinishCurrentFrame > 0)
-		{
-			// Finish current frame immediately to reduce latency
-			Viewport->IssueFrameEvent();
-			Viewport->WaitForFrameEventCompletion();
-		}
-		else if (FinishCurrentFrame < 0)
-		{
-			//No frame wait
-		}
+			if (bNativelyPresented || GUseSubmissionThreadPresent)
+			{
+				FRenderThreadIdleScope Scope(ERenderThreadIdleTypes::WaitingForAllOtherSleep);
+				static const auto CFinishFrameVar = IConsoleManager::Get().FindTConsoleVariableDataInt(TEXT("r.FinishCurrentFrame"));
 
-	}
+				int32 FinishCurrentFrame = CFinishFrameVar->GetValueOnRenderThread();
 
-	// If the input latency timer has been triggered, block until the GPU is completely
-	// finished displaying this frame and calculate the delta time.
-	if (GInputLatencyTimer.RenderThreadTrigger)
-	{
-		Viewport->WaitForFrameEventCompletion();
-		uint32 EndTime = FPlatformTime::Cycles();
-		GInputLatencyTimer.DeltaTime = EndTime - GInputLatencyTimer.StartTime;
-		GInputLatencyTimer.RenderThreadTrigger = false;
-	}
-}
+				if (FinishCurrentFrame == 0)
+				{
+					// Wait for the GPU to finish rendering the previous frame before finishing this frame.
+					Viewport->WaitForFrameEventCompletion();
+					Viewport->IssueFrameEvent();
+				}
+				else if (FinishCurrentFrame > 0)
+				{
+					// Finish current frame immediately to reduce latency
+					Viewport->IssueFrameEvent();
+					Viewport->WaitForFrameEventCompletion();
+				}
+				else if (FinishCurrentFrame < 0)
+				{
+					//No frame wait
+				}
+			}
 
-void FD3D12DynamicRHI::RHIAdvanceFrameForGetViewportBackBuffer(FRHIViewport* ViewportRHI, bool bPresent)
-{
-	check(IsInRenderingThread());
+			// If the input latency timer has been triggered, block until the GPU is completely
+			// finished displaying this frame and calculate the delta time.
+			if (GInputLatencyTimer.RenderThreadTrigger)
+			{
+				Viewport->WaitForFrameEventCompletion();
+				uint32 EndTime = FPlatformTime::Cycles();
+				GInputLatencyTimer.DeltaTime = EndTime - GInputLatencyTimer.StartTime;
+				GInputLatencyTimer.RenderThreadTrigger = false;
+			}
+		});
+	}
 
 #if !UE_BUILD_SHIPPING
 	if (RHIConsoleVariables::LogViewportEvents)
 	{
 		const FString& ThreadName = FThreadManager::GetThreadName(FPlatformTLS::GetCurrentThreadId());
-		UE_LOG(LogD3D12RHI, Log, TEXT("Thread %s: RHIAdvanceFrameForGetViewportBackBuffer"), ThreadName.GetCharArray().GetData());
+		UE_LOG(LogD3D12RHI, Log, TEXT("Thread %s: RHIEndDrawingViewport"), ThreadName.GetCharArray().GetData());
 	}
 #endif
 
@@ -958,8 +972,7 @@ void FD3D12DynamicRHI::RHIAdvanceFrameForGetViewportBackBuffer(FRHIViewport* Vie
 	// But advance the expected present GPU index so the next call to RHIGetViewportNextPresentGPUIndex returns the expected GPU index for the next present.
 	// Warning: when present fails or is not called on the RHIThread then this might not be in sync but RHI thread will fix up the correct state
 	//          Present doesn't happen so shouldn't matter that the index was wrong then
-	FD3D12Viewport* Viewport = FD3D12DynamicRHI::ResourceCast(ViewportRHI);
-	Viewport->AdvanceExpectedBackBufferIndex_RenderThread();
+	ResourceCast(ViewportRHI)->AdvanceExpectedBackBufferIndex_RenderThread();
 }
 
 uint32 FD3D12DynamicRHI::RHIGetViewportNextPresentGPUIndex(FRHIViewport* ViewportRHI)
diff --git a/Engine/Source/Runtime/NullDrv/Public/NullRHI.h b/Engine/Source/Runtime/NullDrv/Public/NullRHI.h
--- a/Engine/Source/Runtime/NullDrv/Public/NullRHI.h
+++ b/Engine/Source/Runtime/NullDrv/Public/NullRHI.h
@@ -409,7 +409,7 @@ public:
 		return true; 
 	}
 
-	virtual void RHIEndDrawingViewport(FRHIViewport* Viewport, const FRHIPresentArgs& InPresentArgs) final override
+	virtual void RHIEndDrawingViewport(FRHICommandListImmediate& RHICmdList, FRHIViewport* Viewport, const FRHIPresentArgs& PresentArgs) final override
 	{
 	}
 
@@ -427,10 +427,6 @@ public:
 	virtual void RHIAliasTextureResources(FTextureRHIRef& DestTexture, FTextureRHIRef& SrcTexture) final override
 	{
 
-	}
-	virtual void RHIAdvanceFrameForGetViewportBackBuffer(FRHIViewport* Viewport, bool bPresent) final override
-	{
-
 	}
 
 	virtual void RHIFlushResources() final override
diff --git a/Engine/Source/Runtime/OpenGLDrv/Private/OpenGLDrv.h b/Engine/Source/Runtime/OpenGLDrv/Private/OpenGLDrv.h
--- a/Engine/Source/Runtime/OpenGLDrv/Private/OpenGLDrv.h
+++ b/Engine/Source/Runtime/OpenGLDrv/Private/OpenGLDrv.h
@@ -159,7 +159,6 @@ public:
 	UE_API virtual FTextureRHIRef RHIGetViewportBackBuffer(FRHIViewport* Viewport) final override;
 	UE_API virtual void RHIAliasTextureResources(FTextureRHIRef& DestTextureRHI, FTextureRHIRef& SrcTextureRHI) final override;
 	UE_API virtual FTextureRHIRef RHICreateAliasedTexture(FTextureRHIRef& SourceTexture) final override;
-	UE_API virtual void RHIAdvanceFrameForGetViewportBackBuffer(FRHIViewport* Viewport, bool bPresent) final override;
 	UE_API virtual void RHIAcquireThreadOwnership() final override;
 	UE_API virtual void RHIReleaseThreadOwnership() final override;
 	UE_API virtual void RHIFlushResources() final override;
@@ -195,7 +194,7 @@ public:
 	UE_API virtual void RHIBeginRenderQuery(FRHIRenderQuery* RenderQuery) final override;
 	UE_API virtual void RHIEndRenderQuery(FRHIRenderQuery* RenderQuery) final override;
 
-	UE_API virtual void RHIEndDrawingViewport(FRHIViewport* Viewport, const FRHIPresentArgs& InPresentArgs) final override;
+	UE_API virtual void RHIEndDrawingViewport(FRHICommandListImmediate& RHICmdList, FRHIViewport* Viewport, const FRHIPresentArgs& PresentArgs) final override;
 	UE_API virtual void RHIEndFrame(const FRHIEndFrameArgs& Args) final override;
 	UE_API virtual void RHISetStreamSource(uint32 StreamIndex, FRHIBuffer* VertexBuffer, uint32 Offset) final override;
 	UE_API virtual void RHISetRasterizerState(FRHIRasterizerState* NewState) final override;
diff --git a/Engine/Source/Runtime/OpenGLDrv/Private/OpenGLViewport.cpp b/Engine/Source/Runtime/OpenGLDrv/Private/OpenGLViewport.cpp
--- a/Engine/Source/Runtime/OpenGLDrv/Private/OpenGLViewport.cpp
+++ b/Engine/Source/Runtime/OpenGLDrv/Private/OpenGLViewport.cpp
@@ -109,85 +109,87 @@ void FOpenGLDynamicRHI::RHITick( float DeltaTime )
  *	Viewport functions.
  *=============================================================================*/
 
-void FOpenGLDynamicRHI::RHIEndDrawingViewport(FRHIViewport* ViewportRHI, const FRHIPresentArgs& InPresentArgs)
+void FOpenGLDynamicRHI::RHIEndDrawingViewport(FRHICommandListImmediate& RHICmdList, FRHIViewport* ViewportRHI, const FRHIPresentArgs& PresentArgs)
 {
-	VERIFY_GL_SCOPE();
-
-	FOpenGLViewport* Viewport = ResourceCast(ViewportRHI);
-
-	SCOPE_CYCLE_COUNTER(STAT_OpenGLPresentTime);
+	//
+	// Since OpenGL RHI does not implement parallel translation, we can enqueue the present from within an RHI command lambda.
+	// This will run after all prior work that has already been recorded by the renderer, on the RHI thread.
+	//
+	RHICmdList.EnqueueLambda(TEXT("RHIEndDrawingViewport"), [this, ViewportRHI, PresentArgs](FRHICommandListImmediate&)
 	{
-		FRenderThreadIdleScope IdleScope(ERenderThreadIdleTypes::WaitingForGPUPresent);
+		VERIFY_GL_SCOPE();
 
-		FOpenGLTexture* BackBuffer = Viewport->GetBackBuffer();
+		FOpenGLViewport* Viewport = ResourceCast(ViewportRHI);
 
-		if (ContextState.bScissorEnabled)
+		SCOPE_CYCLE_COUNTER(STAT_OpenGLPresentTime);
 		{
-			ContextState.bScissorEnabled = false;
-			glDisable(GL_SCISSOR_TEST);
-		}
+			FRenderThreadIdleScope IdleScope(ERenderThreadIdleTypes::WaitingForGPUPresent);
 
-		bool bNeedFinishFrame = PlatformBlitToViewport(*this, PlatformDevice,
-			*Viewport, 
-			BackBuffer->GetSizeX(),
-			BackBuffer->GetSizeY(),
-			InPresentArgs.bPresent,
-			InPresentArgs.bLockToVsync
-		);
+			FOpenGLTexture* BackBuffer = Viewport->GetBackBuffer();
 
-		// Always consider the Framebuffer in the rendering context dirty after the blit
-		ContextState.Framebuffer = -1;
+			if (ContextState.bScissorEnabled)
+			{
+				ContextState.bScissorEnabled = false;
+				glDisable(GL_SCISSOR_TEST);
+			}
 
-		if (bNeedFinishFrame)
-		{
-			static const auto CFinishFrameVar = IConsoleManager::Get().FindTConsoleVariableDataInt(TEXT("r.FinishCurrentFrame"));
-			int32 FinishCurrentFrame = CFinishFrameVar->GetValueOnRenderThread();
+			bool bNeedFinishFrame = PlatformBlitToViewport(*this, PlatformDevice,
+				*Viewport, 
+				BackBuffer->GetSizeX(),
+				BackBuffer->GetSizeY(),
+				PresentArgs.bPresent,
+				PresentArgs.bLockToVsync
+			);
+
+			// Always consider the Framebuffer in the rendering context dirty after the blit
+			ContextState.Framebuffer = -1;
 
-			if (FinishCurrentFrame == 0)
+			if (bNeedFinishFrame)
 			{
-				// Wait for the GPU to finish rendering the previous frame before finishing this frame.
-				Viewport->WaitForFrameEventCompletion();
-				Viewport->IssueFrameEvent();
+				static const auto CFinishFrameVar = IConsoleManager::Get().FindTConsoleVariableDataInt(TEXT("r.FinishCurrentFrame"));
+				int32 FinishCurrentFrame = CFinishFrameVar->GetValueOnRenderThread();
+
+				if (FinishCurrentFrame == 0)
+				{
+					// Wait for the GPU to finish rendering the previous frame before finishing this frame.
+					Viewport->WaitForFrameEventCompletion();
+					Viewport->IssueFrameEvent();
+				}
+				else if (FinishCurrentFrame > 0)
+				{
+					// Finish current frame immediately to reduce latency
+					Viewport->IssueFrameEvent();
+					Viewport->WaitForFrameEventCompletion();
+				}
+				else if (FinishCurrentFrame < 0)
+				{
+					//No frame wait
+				}
 			}
-			else if (FinishCurrentFrame > 0)
+		
+			// If the input latency timer has been triggered, block until the GPU is completely
+			// finished displaying this frame and calculate the delta time.
+			if ( GInputLatencyTimer.RenderThreadTrigger )
 			{
-				// Finish current frame immediately to reduce latency
-				Viewport->IssueFrameEvent();
 				Viewport->WaitForFrameEventCompletion();
+				uint32 EndTime = FPlatformTime::Cycles();
+				GInputLatencyTimer.DeltaTime = EndTime - GInputLatencyTimer.StartTime;
+				GInputLatencyTimer.RenderThreadTrigger = false;
 			}
-			else if (FinishCurrentFrame < 0)
-			{
-				//No frame wait
-			}
-		}
-		
-		// If the input latency timer has been triggered, block until the GPU is completely
-		// finished displaying this frame and calculate the delta time.
-		if ( GInputLatencyTimer.RenderThreadTrigger )
-		{
-			Viewport->WaitForFrameEventCompletion();
-			uint32 EndTime = FPlatformTime::Cycles();
-			GInputLatencyTimer.DeltaTime = EndTime - GInputLatencyTimer.StartTime;
-			GInputLatencyTimer.RenderThreadTrigger = false;
 		}
-	}
 
-	EndFrameTick();
+		EndFrameTick();
 
-	CheckForGLErrors();
+		CheckForGLErrors();
+	});
 }
 
-
 FTextureRHIRef FOpenGLDynamicRHI::RHIGetViewportBackBuffer(FRHIViewport* ViewportRHI)
 {
 	FOpenGLViewport* Viewport = ResourceCast(ViewportRHI);
 	return Viewport->GetBackBuffer();
 }
 
-void FOpenGLDynamicRHI::RHIAdvanceFrameForGetViewportBackBuffer(FRHIViewport* Viewport, bool bPresent)
-{
-}
-
 
 FOpenGLViewport::FOpenGLViewport(FOpenGLDynamicRHI* InOpenGLRHI,void* InWindowHandle,uint32 InSizeX,uint32 InSizeY,bool bInIsFullscreen,EPixelFormat PreferredPixelFormat)
 	: OpenGLRHI(InOpenGLRHI)
diff --git a/Engine/Source/Runtime/RHI/Private/RHICommandList.cpp b/Engine/Source/Runtime/RHI/Private/RHICommandList.cpp
--- a/Engine/Source/Runtime/RHI/Private/RHICommandList.cpp
+++ b/Engine/Source/Runtime/RHI/Private/RHICommandList.cpp
@@ -989,15 +989,18 @@ void FRHICommandListExecutor::FSubmitState::Dispatch(FRHICommandListBase* CmdLis
 			GRHICommandList.LastMutate = CmdList->LastLockFenceCommand->Fence;
 		}
 
-		//
-		// Don't start new translations until all prior submissions have been made.
-		// This is because work inside RHIEndFrame needs to complete on the RHI thread before any new translates can start.
-		// 
-		// Also some RHI commands directly submit to GPU queues from within the RHI (during RHICmdList translation).
-		// Not waiting for prior submits means these internal submissions can happen out-of-order with respect to other translations.
-		// E.g. some platform implementations of EndDrawingViewport() submit directly to the GPU to perform the flip / present.
-		//
-		Prereqs.Add(GRHICommandList.LastSubmit);
+		if (!GRHIGlobals.SupportsConcurrentTranslateAndSubmit)
+		{
+			//
+			// Don't start new translations until all prior submissions have been made.
+			// This is because work inside RHIEndFrame needs to complete on the RHI thread before any new translates can start.
+			// 
+			// Also some RHI commands directly submit to GPU queues from within the RHI (during RHICmdList translation).
+			// Not waiting for prior submits means these internal submissions can happen out-of-order with respect to other translations.
+			// E.g. some platform implementations of EndDrawingViewport() submit directly to the GPU to perform the flip / present.
+			//
+			Prereqs.Add(GRHICommandList.LastSubmit);
+		}
 
 		CurrentTranslateJob->EnqueueTranslateTask(MoveTemp(Prereqs),
 			[TranslateJob = CurrentTranslateJob, CmdList]()
@@ -1858,28 +1861,11 @@ void FRHICommandListImmediate::UnStallRHIThread()
 
 void FRHICommandListImmediate::EndDrawingViewport(FRHIViewport* Viewport, const FRHIPresentArgs& InPresentArgs)
 {
-	// Make sure all prior graphics and async compute work has been submitted.
-	// This is necessary because platform RHIs often submit additional work on the graphics queue during present, and we need to ensure we won't deadlock on async work that wasn't yet submitted by the renderer.
-	// In future, Present() itself should be an enqueued / recorded command, and platform RHIs should never implicitly submit graphics or async compute work.
-	ImmediateFlush(EImmediateFlushType::DispatchToRHIThread);
-
-	check(IsImmediate() && IsInRenderingThread());
-	if (Bypass())
-	{
-		GetContext().RHIEndDrawingViewport(Viewport, InPresentArgs);
-	}
-	else
-	{
-		ALLOC_COMMAND(FRHICommandEndDrawingViewport)(Viewport, InPresentArgs);
-
-		// if we aren't running an RHIThread, there is no good reason to buffer this frame advance stuff and that complicates state management, so flush everything out now
-		{
-			QUICK_SCOPE_CYCLE_COUNTER(STAT_EndDrawingViewport_Dispatch);
-			ImmediateFlush(EImmediateFlushType::DispatchToRHIThread);
-		}
-	}
+	RHI_BREADCRUMB_EVENT(*this, "EndDrawingViewport");
+	GDynamicRHI->RHIEndDrawingViewport(*this, Viewport, InPresentArgs);
 
-	GDynamicRHI->RHIAdvanceFrameForGetViewportBackBuffer(Viewport, InPresentArgs.bPresent);
+	// Dispatch to ensure any swapchain present commands enqueued by the platform RHI above are executed in a timely manner.
+	ImmediateFlush(EImmediateFlushType::DispatchToRHIThread);
 }
 
 void FRHICommandListImmediate::EndFrame()
diff --git a/Engine/Source/Runtime/RHI/Public/DynamicRHI.h b/Engine/Source/Runtime/RHI/Public/DynamicRHI.h
--- a/Engine/Source/Runtime/RHI/Public/DynamicRHI.h
+++ b/Engine/Source/Runtime/RHI/Public/DynamicRHI.h
@@ -687,10 +687,6 @@ public:
 	// Check if PSO Initializers are the same used during PSO Precaching (only compare data relevant for the RHI specific PSO)
 	RHI_API virtual bool RHIMatchPrecachePSOInitializers(const FGraphicsPipelineStateInitializer& LHS, const FGraphicsPipelineStateInitializer& RHS);
 
-	// Only relevant with an RHI thread, this advances the backbuffer for the purpose of GetViewportBackBuffer
-	// FlushType: Thread safe
-	virtual void RHIAdvanceFrameForGetViewportBackBuffer(FRHIViewport* Viewport, bool bPresent) = 0;
-
 	//
 	// Acquires or releases ownership of the platform-specific rendering context for the calling thread.
 	// Only required by OpenGL RHI.
@@ -716,6 +712,9 @@ public:
 		RHIResizeViewport(Viewport, SizeX, SizeY, bIsFullscreen);
 	}
 
+	// Enqueues the presentation of the given viewport, and advances the viewport's swapchain to the next available back buffer.
+	virtual void RHIEndDrawingViewport(FRHICommandListImmediate& RHICmdList, FRHIViewport* Viewport, FRHIPresentArgs const& PresentArgs) = 0;
+
 	// Return preferred pixel format if given format is unsupported.
 	virtual EPixelFormat RHIPreferredPixelFormatHint(EPixelFormat PreferredPixelFormat)
 	{
@@ -1282,7 +1281,6 @@ inline uint32 RHIGetHTilePlatformConfig(const FRHITextureDesc& DepthDesc)
 UE_DEPRECATED(5.7, "RHIAdvanceFrameForGetViewportBackBuffer is deprecated. This happens automatically when RHIEndDrawingViewport is called. Remove calls to this function.")
 inline void RHIAdvanceFrameForGetViewportBackBuffer(FRHIViewport* Viewport)
 {
-	return GDynamicRHI->RHIAdvanceFrameForGetViewportBackBuffer(Viewport, true);
 }
 
 RHI_API uint32 RHIGetGPUFrameCycles(uint32 GPUIndex = 0);
diff --git a/Engine/Source/Runtime/RHI/Public/RHICommandList.h b/Engine/Source/Runtime/RHI/Public/RHICommandList.h
--- a/Engine/Source/Runtime/RHI/Public/RHICommandList.h
+++ b/Engine/Source/Runtime/RHI/Public/RHICommandList.h
@@ -2273,19 +2273,6 @@ FRHICOMMAND_MACRO(FRHICommandPostExternalCommandsReset)
 	RHI_EXECUTE_API void Execute(FRHICommandListBase& CmdList);
 };
 
-FRHICOMMAND_MACRO(FRHICommandEndDrawingViewport)
-{
-	FRHIViewport* Viewport;
-	FRHIPresentArgs PresentArgs;
-
-	inline FRHICommandEndDrawingViewport(FRHIViewport* InViewport, const FRHIPresentArgs& InPresentArgs)
-		: Viewport(InViewport)
-		, PresentArgs(InPresentArgs)
-	{
-	}
-	RHI_EXECUTE_API void Execute(FRHICommandListBase& CmdList);
-};
-
 FRHICOMMAND_MACRO(FRHICommandCopyBufferRegion)
 {
 	FRHIBuffer* DestBuffer;
diff --git a/Engine/Source/Runtime/RHI/Public/RHICommandListCommandExecutes.inl b/Engine/Source/Runtime/RHI/Public/RHICommandListCommandExecutes.inl
--- a/Engine/Source/Runtime/RHI/Public/RHICommandListCommandExecutes.inl
+++ b/Engine/Source/Runtime/RHI/Public/RHICommandListCommandExecutes.inl
@@ -28,7 +28,6 @@ struct FRHICommandDrawPrimitive;
 struct FRHICommandDrawPrimitiveIndirect;
 struct FRHICommandMultiDrawPrimitiveIndirect;
 struct FRHICommandSetDepthBounds;
-struct FRHICommandEndDrawingViewport;
 struct FRHICommandEndOcclusionQueryBatch;
 struct FRHICommandEndRenderQuery;
 struct FRHICommandSetBlendFactor;
@@ -539,9 +538,3 @@ void FRHICommandSetBindingsOnShaderBindingTable::Execute(FRHICommandListBase& Cm
 
 	INTERNAL_DECORATOR(RHISetBindingsOnShaderBindingTable)(SBT, RayTracingPipelineState, NumBindings, Bindings, BindingType);
 }
-
-void FRHICommandEndDrawingViewport::Execute(FRHICommandListBase& CmdList)
-{
-	RHISTAT(EndDrawingViewport);
-	INTERNAL_DECORATOR(RHIEndDrawingViewport)(Viewport, PresentArgs);
-}
diff --git a/Engine/Source/Runtime/RHI/Public/RHIContext.h b/Engine/Source/Runtime/RHI/Public/RHIContext.h
--- a/Engine/Source/Runtime/RHI/Public/RHIContext.h
+++ b/Engine/Source/Runtime/RHI/Public/RHIContext.h
@@ -712,9 +712,6 @@ public:
 
 	virtual void RHIEndRenderQuery(FRHIRenderQuery* RenderQuery) = 0;
 
-	// This method is queued with an RHIThread, otherwise it will flush after it is queued; without an RHI thread there is no benefit to queuing this frame advance commands
-	virtual void RHIEndDrawingViewport(FRHIViewport* Viewport, const FRHIPresentArgs& InPresentArgs) = 0;
-
 	virtual void RHISetStreamSource(uint32 StreamIndex, FRHIBuffer* VertexBuffer, uint32 Offset) = 0;
 
 	// @param MinX including like Win32 RECT
diff --git a/Engine/Source/Runtime/RHI/Public/RHIGlobals.h b/Engine/Source/Runtime/RHI/Public/RHIGlobals.h
--- a/Engine/Source/Runtime/RHI/Public/RHIGlobals.h
+++ b/Engine/Source/Runtime/RHI/Public/RHIGlobals.h
@@ -528,6 +528,14 @@ struct FRHIGlobals
 	***/
 	bool SupportsParallelRHIExecute = false;
 	
+	/** 
+	  * If true, RHI submit tasks (those that call RHISubmitCommandLists) are allowed to run concurrently with subsequent RHI translation tasks.
+	  * This is only possible with RHIs that correctly pipeline viewport presentation, and don't implicitly submit GPU work from within RHI translate tasks.
+	  * 
+	  * If false, future RHI translate tasks cannot start until all prior RHI submit tasks have completed. Submit tasks run in isolation on the RHI thread.
+	  */
+	bool SupportsConcurrentTranslateAndSubmit = false;
+
 	/** Enables support for platforms that use render passes to split command list translation across multiple threads. For platforms that require
 	   dependencies on Parent/Child command lists i.e. requiring translation of children to be complete before finishing translation of the
 	   parent, should use ParallelRHIExecuteChildWait and/or ParallelRHIExecuteParentWait */
diff --git a/Engine/Source/Runtime/RHI/Public/RHIValidation.h b/Engine/Source/Runtime/RHI/Public/RHIValidation.h
--- a/Engine/Source/Runtime/RHI/Public/RHIValidation.h
+++ b/Engine/Source/Runtime/RHI/Public/RHIValidation.h
@@ -725,13 +725,6 @@ public:
 		RHI->RHIGetDisplaysInformation(OutDisplayInformation);
 	}
 
-	// Only relevant with an RHI thread, this advances the backbuffer for the purpose of GetViewportBackBuffer
-	// FlushType: Thread safe
-	virtual void RHIAdvanceFrameForGetViewportBackBuffer(FRHIViewport* Viewport, bool bPresent) override final
-	{
-		RHI->RHIAdvanceFrameForGetViewportBackBuffer(Viewport, bPresent);
-	}
-
 	virtual void RHIAcquireThreadOwnership() override final
 	{
 		RHI->RHIAcquireThreadOwnership();
@@ -769,6 +762,11 @@ public:
 		RHI->RHIResizeViewport(Viewport, SizeX, SizeY, bIsFullscreen, PreferredPixelFormat);
 	}
 
+	virtual void RHIEndDrawingViewport(FRHICommandListImmediate& RHICmdList, FRHIViewport* Viewport, const FRHIPresentArgs& PresentArgs) override final
+	{
+		RHI->RHIEndDrawingViewport(RHICmdList, Viewport, PresentArgs);
+	}
+
 	virtual EPixelFormat RHIPreferredPixelFormatHint(EPixelFormat PreferredPixelFormat) override final
 	{
 		return RHI->RHIPreferredPixelFormatHint(PreferredPixelFormat);
diff --git a/Engine/Source/Runtime/RHI/Public/RHIValidationContext.h b/Engine/Source/Runtime/RHI/Public/RHIValidationContext.h
--- a/Engine/Source/Runtime/RHI/Public/RHIValidationContext.h
+++ b/Engine/Source/Runtime/RHI/Public/RHIValidationContext.h
@@ -758,12 +758,6 @@ public:
 		RHIContext->RHIEndRenderQuery(RenderQuery);
 	}
 
-	// This method is queued with an RHIThread, otherwise it will flush after it is queued; without an RHI thread there is no benefit to queuing this frame advance commands
-	virtual void RHIEndDrawingViewport(FRHIViewport* Viewport, const FRHIPresentArgs& InPresentArgs) override final
-	{
-		RHIContext->RHIEndDrawingViewport(Viewport, InPresentArgs);
-	}
-
 	virtual void RHISetStreamSource(uint32 StreamIndex, FRHIBuffer* VertexBuffer, uint32 Offset) override final
 	{
 		//#todo-rco: Decide if this is needed or not...
diff --git a/Engine/Source/Runtime/VulkanRHI/Private/VulkanContext.h b/Engine/Source/Runtime/VulkanRHI/Private/VulkanContext.h
--- a/Engine/Source/Runtime/VulkanRHI/Private/VulkanContext.h
+++ b/Engine/Source/Runtime/VulkanRHI/Private/VulkanContext.h
@@ -319,8 +319,6 @@ public:
 
 	virtual void Finalize(TArray<FVulkanPayload*>& OutPayloads) override;
 
-	virtual void RHIEndDrawingViewport(FRHIViewport* Viewport, const FRHIPresentArgs& InPresentArgs) final override;
-
 	virtual void RHIBeginRenderPass(const FRHIRenderPassInfo& InInfo, const TCHAR* InName) final override;
 	virtual void RHIEndRenderPass() final override;
 	virtual void RHINextSubpass() final override;
diff --git a/Engine/Source/Runtime/VulkanRHI/Private/VulkanDynamicRHI.h b/Engine/Source/Runtime/VulkanRHI/Private/VulkanDynamicRHI.h
--- a/Engine/Source/Runtime/VulkanRHI/Private/VulkanDynamicRHI.h
+++ b/Engine/Source/Runtime/VulkanRHI/Private/VulkanDynamicRHI.h
@@ -209,11 +209,11 @@ public:
 	virtual FTextureRHIRef RHIGetViewportBackBuffer(FRHIViewport* Viewport) final override;
 	virtual void RHIAliasTextureResources(FTextureRHIRef& DestTexture, FTextureRHIRef& SrcTexture) final override;
 	virtual FTextureRHIRef RHICreateAliasedTexture(FTextureRHIRef& SourceTexture) final override;
-	virtual void RHIAdvanceFrameForGetViewportBackBuffer(FRHIViewport* Viewport, bool bPresent) final override;
 	virtual void RHIFlushResources() final override;
 	virtual FViewportRHIRef RHICreateViewport(void* WindowHandle, uint32 SizeX, uint32 SizeY, bool bIsFullscreen, EPixelFormat PreferredPixelFormat) final override;
 	virtual void RHIResizeViewport(FRHIViewport* Viewport, uint32 SizeX, uint32 SizeY, bool bIsFullscreen) final override;
 	virtual void RHIResizeViewport(FRHIViewport* Viewport, uint32 SizeX, uint32 SizeY, bool bIsFullscreen, EPixelFormat PreferredPixelFormat) final override;
+	virtual void RHIEndDrawingViewport(FRHICommandListImmediate& RHICmdList, FRHIViewport* Viewport, const FRHIPresentArgs& PresentArgs) final override;
 	virtual void RHITick(float DeltaTime) final override;
 	virtual void RHIBlockUntilGPUIdle() final override;
 	virtual void RHISuspendRendering() final override;
diff --git a/Engine/Source/Runtime/VulkanRHI/Private/VulkanRHI.cpp b/Engine/Source/Runtime/VulkanRHI/Private/VulkanRHI.cpp
--- a/Engine/Source/Runtime/VulkanRHI/Private/VulkanRHI.cpp
+++ b/Engine/Source/Runtime/VulkanRHI/Private/VulkanRHI.cpp
@@ -1249,21 +1249,42 @@ FVulkanTiming::FVulkanTiming(FVulkanQueue& InQueue)
 {
 }
 
-void FVulkanCommandListContext::RHIEndDrawingViewport(FRHIViewport* ViewportRHI, const FRHIPresentArgs& InPresentArgs)
-{
-	LLM_SCOPE_VULKAN(ELLMTagVulkan::VulkanMisc);
-	//FRCLog::Printf(FString::Printf(TEXT("FVulkanCommandListContext::RHIEndDrawingViewport()")));
-	check(IsImmediate());
-	FVulkanViewport* Viewport = ResourceCast(ViewportRHI);
-
-	FVulkanCommandBuffer& CommandBuffer = GetCommandBuffer();
-	check(!CommandBuffer.HasEnded() && !CommandBuffer.IsInsideRenderPass());
+void FVulkanDynamicRHI::RHIEndDrawingViewport(FRHICommandListImmediate& RHICmdList, FRHIViewport* ViewportRHI, const FRHIPresentArgs& PresentArgs)
+{
+	//
+	// Make sure all prior graphics and async compute work has been submitted. This is necessary because Vulkan RHI handles the present from inside the below lambda.
+	// Since the present is not yet pipelined via payloads etc, we must ensure prior work has been handed to RHISubmitCommandLists before we translate the lambda.
+	//
+	// In future, Present() itself should be an enqueued command, and platform RHIs should never implicitly send commands to the GPU queues during RHI translation.
+	// All GPU work should be submitted via RHISubmitCommandLists.
+	//
+	// This is also why there is a task dependency between RHI translate and RHI submission tasks in RHICommandList.cpp.
+	// See GRHIGlobals.SupportsConcurrentTranslateAndSubmit
+	//
+	RHICmdList.ImmediateFlush(EImmediateFlushType::DispatchToRHIThread);
 
-	const bool bNativePresent = Viewport->Present(*this, InPresentArgs.bLockToVsync);
-	if (bNativePresent)
 	{
-		//#todo-rco: Check for r.FinishCurrentFrame
+		FRHICommandListScopedPipeline PipelineScope(RHICmdList, ERHIPipeline::Graphics);
+		RHICmdList.EnqueueLambda(TEXT("RHIEndDrawingViewport"), [ViewportRHI, PresentArgs](FRHICommandListImmediate& ExecutingCmdList)
+		{
+			LLM_SCOPE_VULKAN(ELLMTagVulkan::VulkanMisc);
+			FVulkanViewport* Viewport = ResourceCast(ViewportRHI);
+
+			FVulkanCommandListContext& Context = FVulkanCommandListContext::Get(ExecutingCmdList);
+
+			FVulkanCommandBuffer& CommandBuffer = Context.GetCommandBuffer();
+			check(!CommandBuffer.HasEnded() && !CommandBuffer.IsInsideRenderPass());
+
+			const bool bNativePresent = Viewport->Present(Context, PresentArgs.bLockToVsync);
+			if (bNativePresent)
+			{
+				//#todo-rco: Check for r.FinishCurrentFrame
+			}
+		});
 	}
+
+	// Switch to the next viewport back buffer
+	ResourceCast(ViewportRHI)->AdvanceBackBufferFrame(RHICmdList);
 }
 
 #if WITH_RHI_BREADCRUMBS
diff --git a/Engine/Source/Runtime/VulkanRHI/Private/VulkanViewport.cpp b/Engine/Source/Runtime/VulkanRHI/Private/VulkanViewport.cpp
--- a/Engine/Source/Runtime/VulkanRHI/Private/VulkanViewport.cpp
+++ b/Engine/Source/Runtime/VulkanRHI/Private/VulkanViewport.cpp
@@ -940,16 +940,6 @@ FTextureRHIRef FVulkanDynamicRHI::RHIGetViewportBackBuffer(FRHIViewport* Viewpor
 	return Viewport->GetBackBuffer(RHICmdList);
 }
 
-void FVulkanDynamicRHI::RHIAdvanceFrameForGetViewportBackBuffer(FRHIViewport* ViewportRHI, bool bPresent)
-{
-	FRHICommandListImmediate& RHICmdList = FRHICommandListImmediate::Get();
-
-	check(ViewportRHI);
-	FVulkanViewport* Viewport = ResourceCast(ViewportRHI);
-
-	Viewport->AdvanceBackBufferFrame(RHICmdList);
-}
-
 void FVulkanCommandListContext::RHISetViewport(float MinX, float MinY, float MinZ, float MaxX, float MaxY, float MaxZ)
 {
 	PendingGfxState->SetViewport(MinX, MinY, MinZ, MaxX, MaxY, MaxZ);
diff --git a/Engine/Source/Runtime/Windows/D3D11RHI/Private/D3D11Device.cpp b/Engine/Source/Runtime/Windows/D3D11RHI/Private/D3D11Device.cpp
--- a/Engine/Source/Runtime/Windows/D3D11RHI/Private/D3D11Device.cpp
+++ b/Engine/Source/Runtime/Windows/D3D11RHI/Private/D3D11Device.cpp
@@ -53,7 +53,6 @@ FD3D11DynamicRHI::FD3D11DynamicRHI(IDXGIFactory1* InDXGIFactory1, D3D_FEATURE_LE
 	, CurrentDepthTexture(NULL)
 	, NumSimultaneousRenderTargets(0)
 	, NumUAVs(0)
-	, PresentCounter(0)
 	, CurrentDSVAccessType(FExclusiveDepthStencil::DepthWrite_StencilWrite)
 	, bDiscardSharedConstants(false)
 	, Adapter(InAdapter)
diff --git a/Engine/Source/Runtime/Windows/D3D11RHI/Private/D3D11RHIPrivate.h b/Engine/Source/Runtime/Windows/D3D11RHI/Private/D3D11RHIPrivate.h
--- a/Engine/Source/Runtime/Windows/D3D11RHI/Private/D3D11RHIPrivate.h
+++ b/Engine/Source/Runtime/Windows/D3D11RHI/Private/D3D11RHIPrivate.h
@@ -290,7 +290,6 @@ public:
 	D3D11RHI_API virtual FTextureRHIRef RHIGetViewportBackBuffer(FRHIViewport* Viewport) final override;
 	D3D11RHI_API virtual void RHIAliasTextureResources(FTextureRHIRef& DestTexture, FTextureRHIRef& SrcTexture) final override;
 	D3D11RHI_API virtual FTextureRHIRef RHICreateAliasedTexture(FTextureRHIRef& SourceTexture) final override;
-	D3D11RHI_API virtual void RHIAdvanceFrameForGetViewportBackBuffer(FRHIViewport* Viewport, bool bPresent) final override;
 	D3D11RHI_API virtual void RHIFlushResources() final override;
 	D3D11RHI_API virtual FViewportRHIRef RHICreateViewport(void* WindowHandle, uint32 SizeX, uint32 SizeY, bool bIsFullscreen, EPixelFormat PreferredPixelFormat) override;
 	D3D11RHI_API virtual void RHIResizeViewport(FRHIViewport* Viewport, uint32 SizeX, uint32 SizeY, bool bIsFullscreen) final override;
@@ -329,7 +328,7 @@ public:
 	D3D11RHI_API virtual void RHIBeginRenderQuery(FRHIRenderQuery* RenderQuery) final override;
 	D3D11RHI_API virtual void RHIEndRenderQuery  (FRHIRenderQuery* RenderQuery) final override;
 
-	D3D11RHI_API virtual void RHIEndDrawingViewport(FRHIViewport* Viewport, const FRHIPresentArgs& InPresentArgs) final override;
+	D3D11RHI_API virtual void RHIEndDrawingViewport(FRHICommandListImmediate& RHICmdList, FRHIViewport* Viewport, const FRHIPresentArgs& PresentArgs) final override;
 	D3D11RHI_API virtual void RHIEndFrame(const FRHIEndFrameArgs& Args) final override;
 	D3D11RHI_API virtual void RHISetStreamSource(uint32 StreamIndex, FRHIBuffer* VertexBuffer, uint32 Offset) final override;
 	D3D11RHI_API virtual void RHISetRasterizerState(FRHIRasterizerState* NewState) final override;
@@ -572,9 +571,6 @@ protected:
 	uint32 NumSimultaneousRenderTargets;
 	uint32 NumUAVs;
 
-	/** Internal frame counter that just counts calls to Present */
-	uint32 PresentCounter;
-
 	// Render queries that should be polled by the RHI thread.
 	struct
 	{
diff --git a/Engine/Source/Runtime/Windows/D3D11RHI/Private/D3D11Viewport.cpp b/Engine/Source/Runtime/Windows/D3D11RHI/Private/D3D11Viewport.cpp
--- a/Engine/Source/Runtime/Windows/D3D11RHI/Private/D3D11Viewport.cpp
+++ b/Engine/Source/Runtime/Windows/D3D11RHI/Private/D3D11Viewport.cpp
@@ -780,88 +780,91 @@ void FD3D11DynamicRHI::RHITick( float DeltaTime )
  *	Viewport functions.
  *=============================================================================*/
 
-void FD3D11DynamicRHI::RHIEndDrawingViewport(FRHIViewport* ViewportRHI, const FRHIPresentArgs& InPresentArgs)
+void FD3D11DynamicRHI::RHIEndDrawingViewport(FRHICommandListImmediate& RHICmdList, FRHIViewport* ViewportRHI, const FRHIPresentArgs& PresentArgs)
 {
-	++PresentCounter;
-	FD3D11Viewport* Viewport = ResourceCast(ViewportRHI);
+	//
+	// Since D3D11 RHI does not implement parallel translation, we can enqueue the present from within an RHI command lambda.
+	// This will run after all prior work that has already been recorded by the renderer, on the RHI thread.
+	//
+	RHICmdList.EnqueueLambda(TEXT("RHIEndDrawingViewport"), [this, ViewportRHI, PresentArgs](FRHICommandListImmediate&)
+	{
+		FD3D11Viewport* Viewport = ResourceCast(ViewportRHI);
 
-	SCOPE_CYCLE_COUNTER(STAT_D3D11PresentTime);
+		SCOPE_CYCLE_COUNTER(STAT_D3D11PresentTime);
 
-	// Clear references the device might have to resources.
-	CurrentDepthTexture = NULL;
-	CurrentDepthStencilTarget = NULL;
-	CurrentDSVAccessType = FExclusiveDepthStencil::DepthWrite_StencilWrite;
-	CurrentRenderTargets[0] = NULL;
-	for(uint32 RenderTargetIndex = 1;RenderTargetIndex < D3D11_SIMULTANEOUS_RENDER_TARGET_COUNT;++RenderTargetIndex)
-	{
-		CurrentRenderTargets[RenderTargetIndex] = NULL;
-	}
+		// Clear references the device might have to resources.
+		CurrentDepthTexture = nullptr;
+		CurrentDepthStencilTarget = nullptr;
+		CurrentDSVAccessType = FExclusiveDepthStencil::DepthWrite_StencilWrite;
+		CurrentRenderTargets[0] = nullptr;
 
-	ClearAllShaderResources();
+		for (uint32 RenderTargetIndex = 1; RenderTargetIndex < D3D11_SIMULTANEOUS_RENDER_TARGET_COUNT; ++RenderTargetIndex)
+		{
+			CurrentRenderTargets[RenderTargetIndex] = nullptr;
+		}
 
-	CommitRenderTargetsAndUAVs();
+		ClearAllShaderResources();
 
-	StateCache.SetVertexShader(nullptr);
+		CommitRenderTargetsAndUAVs();
 
-	uint16 NullStreamStrides[MaxVertexElementCount] = {0};
-	StateCache.SetStreamStrides(NullStreamStrides);
-	for (uint32 StreamIndex = 0; StreamIndex < MaxVertexElementCount; ++StreamIndex)
-	{
-		StateCache.SetStreamSource(nullptr, StreamIndex, 0, 0);
-	}
+		StateCache.SetVertexShader(nullptr);
 
-	StateCache.SetIndexBuffer(nullptr, DXGI_FORMAT_R16_UINT, 0);
+		uint16 NullStreamStrides[MaxVertexElementCount] = {0};
+		StateCache.SetStreamStrides(NullStreamStrides);
+		for (uint32 StreamIndex = 0; StreamIndex < MaxVertexElementCount; ++StreamIndex)
+		{
+			StateCache.SetStreamSource(nullptr, StreamIndex, 0, 0);
+		}
+
+		StateCache.SetIndexBuffer(nullptr, DXGI_FORMAT_R16_UINT, 0);
 	
-	CurrentResourceBoundAsIB = nullptr;
-	FMemory::Memzero(CurrentResourcesBoundAsVBs, sizeof(CurrentResourcesBoundAsVBs));
-	MaxBoundVertexBufferIndex = INDEX_NONE;
+		CurrentResourceBoundAsIB = nullptr;
+		FMemory::Memzero(CurrentResourcesBoundAsVBs, sizeof(CurrentResourcesBoundAsVBs));
+		MaxBoundVertexBufferIndex = INDEX_NONE;
 	
-	StateCache.SetPixelShader(nullptr);
-	StateCache.SetGeometryShader(nullptr);
-	// Compute Shader is set to NULL after each Dispatch call, so no need to clear it here
+		StateCache.SetPixelShader(nullptr);
+		StateCache.SetGeometryShader(nullptr);
+		// Compute Shader is set to nullptr after each Dispatch call, so no need to clear it here
 
-	bool bNativelyPresented = true;
-	if (InPresentArgs.bPresent)
-	{
-		bNativelyPresented = Viewport->Present(*this, InPresentArgs.bLockToVsync);
-	}
+		bool bNativelyPresented = true;
+		if (PresentArgs.bPresent)
+		{
+			bNativelyPresented = Viewport->Present(*this, PresentArgs.bLockToVsync);
+		}
 
-	if (bNativelyPresented)
-	{ 
-		static const auto CFinishFrameVar = IConsoleManager::Get().FindTConsoleVariableDataInt(TEXT("r.FinishCurrentFrame"));
-		int32 FinishCurrentFrame = CFinishFrameVar->GetValueOnRenderThread();
+		if (bNativelyPresented)
+		{ 
+			static const auto CFinishFrameVar = IConsoleManager::Get().FindTConsoleVariableDataInt(TEXT("r.FinishCurrentFrame"));
+			int32 FinishCurrentFrame = CFinishFrameVar->GetValueOnRenderThread();
 
-		if (FinishCurrentFrame == 0)
-		{
-			// Wait for the GPU to finish rendering the previous frame before finishing this frame.
-			Viewport->WaitForFrameEventCompletion();
-			Viewport->IssueFrameEvent();
+			if (FinishCurrentFrame == 0)
+			{
+				// Wait for the GPU to finish rendering the previous frame before finishing this frame.
+				Viewport->WaitForFrameEventCompletion();
+				Viewport->IssueFrameEvent();
+			}
+			else if (FinishCurrentFrame > 0)
+			{
+				// Finish current frame immediately to reduce latency
+				Viewport->IssueFrameEvent();
+				Viewport->WaitForFrameEventCompletion();
+			}
+			else if (FinishCurrentFrame < 0)
+			{
+				//No frame wait
+			}
 		}
-		else if (FinishCurrentFrame > 0)
+
+		// If the input latency timer has been triggered, block until the GPU is completely
+		// finished displaying this frame and calculate the delta time.
+		if ( GInputLatencyTimer.RenderThreadTrigger )
 		{
-			// Finish current frame immediately to reduce latency
-			Viewport->IssueFrameEvent();
 			Viewport->WaitForFrameEventCompletion();
+			uint32 EndTime = FPlatformTime::Cycles();
+			GInputLatencyTimer.DeltaTime = EndTime - GInputLatencyTimer.StartTime;
+			GInputLatencyTimer.RenderThreadTrigger = false;
 		}
-		else if (FinishCurrentFrame < 0)
-		{
-			//No frame wait
-		}
-	}
-
-	// If the input latency timer has been triggered, block until the GPU is completely
-	// finished displaying this frame and calculate the delta time.
-	if ( GInputLatencyTimer.RenderThreadTrigger )
-	{
-		Viewport->WaitForFrameEventCompletion();
-		uint32 EndTime = FPlatformTime::Cycles();
-		GInputLatencyTimer.DeltaTime = EndTime - GInputLatencyTimer.StartTime;
-		GInputLatencyTimer.RenderThreadTrigger = false;
-	}
-}
-
-void FD3D11DynamicRHI::RHIAdvanceFrameForGetViewportBackBuffer(FRHIViewport* Viewport, bool bPresent)
-{
+	});
 }
 
 FTextureRHIRef FD3D11DynamicRHI::RHIGetViewportBackBuffer(FRHIViewport* ViewportRHI)
