Remove unused support for multiple adapters from D3D12RHI

#Synced-CL 49917198

diff --git a/Engine/Source/Runtime/D3D12RHI/Private/D3D12CommandContext.cpp b/Engine/Source/Runtime/D3D12RHI/Private/D3D12CommandContext.cpp
--- a/Engine/Source/Runtime/D3D12RHI/Private/D3D12CommandContext.cpp
+++ b/Engine/Source/Runtime/D3D12RHI/Private/D3D12CommandContext.cpp
@@ -702,7 +702,7 @@ void FD3D12DynamicRHI::UpdateMemoryStats()
 #if PLATFORM_WINDOWS && (STATS || CSV_PROFILER_STATS)
 	SCOPE_CYCLE_COUNTER(STAT_D3DUpdateVideoMemoryStats);
 
-	for (TSharedPtr<FD3D12Adapter> const& Adapter : ChosenAdapters)
+	if (FD3D12Adapter* const Adapter = ChosenAdapter.Get())
 	{
 		// Refresh captured memory stats.
 		const FD3DMemoryStats& MemoryStats = Adapter->CollectMemoryStats();
diff --git a/Engine/Source/Runtime/D3D12RHI/Private/D3D12Commands.cpp b/Engine/Source/Runtime/D3D12RHI/Private/D3D12Commands.cpp
--- a/Engine/Source/Runtime/D3D12RHI/Private/D3D12Commands.cpp
+++ b/Engine/Source/Runtime/D3D12RHI/Private/D3D12Commands.cpp
@@ -1455,10 +1455,9 @@ void FD3D12CommandContext::RHIClearMRTImpl(bool* bClearColorArray, int32 NumClea
 // Blocks the CPU until the GPU catches up and goes idle.
 void FD3D12DynamicRHI::RHIBlockUntilGPUIdle()
 {
-	const int32 NumAdapters = ChosenAdapters.Num();
-	for (int32 Index = 0; Index < NumAdapters; ++Index)
+	if (FD3D12Adapter* const Adapter = ChosenAdapter.Get())
 	{
-		GetAdapter(Index).BlockUntilIdle();
+		Adapter->BlockUntilIdle();
 	}
 }
 
diff --git a/Engine/Source/Runtime/D3D12RHI/Private/D3D12RHI.cpp b/Engine/Source/Runtime/D3D12RHI/Private/D3D12RHI.cpp
--- a/Engine/Source/Runtime/D3D12RHI/Private/D3D12RHI.cpp
+++ b/Engine/Source/Runtime/D3D12RHI/Private/D3D12RHI.cpp
@@ -57,8 +57,8 @@ FD3D12WorkaroundFlags GD3D12WorkaroundFlags;
 
 using namespace D3D12RHI;
 
-FD3D12DynamicRHI::FD3D12DynamicRHI(const TArray<TSharedPtr<FD3D12Adapter>>& ChosenAdaptersIn, bool bInPixEventEnabled)
-	: ChosenAdapters(ChosenAdaptersIn)
+FD3D12DynamicRHI::FD3D12DynamicRHI(TSharedPtr<FD3D12Adapter>& InChosenAdapter, bool bInPixEventEnabled)
+	: ChosenAdapter(InChosenAdapter)
 	, bPixEventEnabled(bInPixEventEnabled)
 	, AmdAgsContext(nullptr)
 	, AmdSupportedExtensionFlags(0)
@@ -256,7 +256,7 @@ FD3D12DynamicRHI::FD3D12DynamicRHI(const TArray<TSharedPtr<FD3D12Adapter>>& Chos
 
 void FD3D12DynamicRHI::PostInit()
 {
-	for (TSharedPtr<FD3D12Adapter>& Adapter : ChosenAdapters)
+	if (FD3D12Adapter* const Adapter = ChosenAdapter.Get())
 	{
 		Adapter->InitializeExplicitDescriptorHeap();
 
@@ -271,16 +271,14 @@ FD3D12DynamicRHI::~FD3D12DynamicRHI()
 {
 	UE_LOG(LogD3D12RHI, Log, TEXT("~FD3D12DynamicRHI"));
 
-	check(ChosenAdapters.Num() == 0);
+	check(!ChosenAdapter.IsValid());
 }
 
 void FD3D12DynamicRHI::ForEachQueue(TFunctionRef<void(FD3D12Queue&)> Callback)
 {
-	for (uint32 AdapterIndex = 0; AdapterIndex < GetNumAdapters(); ++AdapterIndex)
+	if (FD3D12Adapter* const Adapter = ChosenAdapter.Get())
 	{
-		FD3D12Adapter& Adapter = GetAdapter(AdapterIndex);
-
-		for (FD3D12Device* Device : Adapter.GetDevices())
+		for (FD3D12Device* const Device : Adapter->GetDevices())
 		{
 			for (FD3D12Queue& Queue : Device->GetQueues())
 			{
@@ -322,7 +320,7 @@ void FD3D12DynamicRHI::Shutdown()
 	// Ask all initialized FRenderResources to release their RHI resources.
 	FRenderResource::ReleaseRHIForAllResources();
 
-	for (TSharedPtr<FD3D12Adapter>& Adapter : ChosenAdapters)
+	if (FD3D12Adapter* const Adapter = ChosenAdapter.Get())
 	{
 		Adapter->CleanupResources();
 		Adapter->BlockUntilIdle();
@@ -337,7 +335,7 @@ void FD3D12DynamicRHI::Shutdown()
 	check(ObjectsToDelete.Num() == 0);
 
 	// Delete adapters, devices, queues, command contexts etc
-	ChosenAdapters.Empty();
+	ChosenAdapter.Reset();
 
 	check(ObjectsToDelete.Num() == 0);
 
@@ -508,12 +506,9 @@ void FD3D12DynamicRHI::RHIProcessDeleteQueue()
 	}
 
 	// Clear all bound resources since we are about to flush pending deletions.
-
-	for (uint32 AdapterIndex = 0; AdapterIndex < GetNumAdapters(); ++AdapterIndex)
+	if (FD3D12Adapter* const Adapter = ChosenAdapter.Get())
 	{
-		FD3D12Adapter& Adapter = GetAdapter(AdapterIndex);
-
-		for (FD3D12Device* Device : Adapter.GetDevices())
+		for (FD3D12Device* const Device : Adapter->GetDevices())
 		{
 			Device->GetDefaultCommandContext().ClearState(FD3D12ContextCommon::EClearStateMode::TransientOnly);
 		}
@@ -522,7 +517,7 @@ void FD3D12DynamicRHI::RHIProcessDeleteQueue()
 
 void FD3D12DynamicRHI::RHIEndFrame_RenderThread(FRHICommandListImmediate& RHICmdList)
 {
-	for (auto& Adapter : ChosenAdapters)
+	if (FD3D12Adapter* const Adapter = ChosenAdapter.Get())
 	{
 		Adapter->GetFrameFence().AdvanceTOP();
 	}
@@ -534,7 +529,7 @@ void FD3D12DynamicRHI::RHIEndFrame_RenderThread(FRHICommandListImmediate& RHICmd
 	RHICmdList.EnqueueLambdaMultiPipe(GetEnabledRHIPipelines(), FRHICommandListBase::EThreadFence::Enabled, TEXT("D3D12 BeginFrame"),
 		[this](FD3D12ContextArray const& Contexts)
 	{
-		for (const TSharedPtr<FD3D12Adapter>& Adapter : ChosenAdapters)
+		if (FD3D12Adapter* const Adapter = ChosenAdapter.Get())
 		{
 			for (FD3D12Device* const Device : Adapter->GetDevices())
 			{
@@ -556,7 +551,7 @@ void FD3D12DynamicRHI::RHIEndFrame_RenderThread(FRHICommandListImmediate& RHICmd
 
 void FD3D12DynamicRHI::RHIEndFrame(const FRHIEndFrameArgs& Args)
 {
-	for (auto& Adapter : ChosenAdapters)
+	if (FD3D12Adapter* const Adapter = ChosenAdapter.Get())
 	{
 		Adapter->EndFrame();
 
@@ -618,7 +613,7 @@ TArray<FD3D12MinimalAdapterDesc> FD3D12DynamicRHI::RHIGetAdapterDescs() const
 {
 	TArray<FD3D12MinimalAdapterDesc> Result;
 
-	for (const TSharedPtr<FD3D12Adapter>& Adapter : ChosenAdapters)
+	if (FD3D12Adapter* const Adapter = ChosenAdapter.Get())
 	{
 		FD3D12MinimalAdapterDesc Desc{};
 		Desc.Desc = Adapter->GetDesc().Desc;
diff --git a/Engine/Source/Runtime/D3D12RHI/Private/D3D12RHIPrivate.h b/Engine/Source/Runtime/D3D12RHI/Private/D3D12RHIPrivate.h
--- a/Engine/Source/Runtime/D3D12RHI/Private/D3D12RHIPrivate.h
+++ b/Engine/Source/Runtime/D3D12RHI/Private/D3D12RHIPrivate.h
@@ -277,7 +277,7 @@ public:
 	void ForEachQueue(TFunctionRef<void(FD3D12Queue&)> Callback);
 
 	/** Initialization constructor. */
-	FD3D12DynamicRHI(const TArray<TSharedPtr<FD3D12Adapter>>& ChosenAdaptersIn, bool bInPixEventEnabled);
+	FD3D12DynamicRHI(TSharedPtr<FD3D12Adapter>& InChosenAdapter, bool bInPixEventEnabled);
 
 	/** Destructor */
 	virtual ~FD3D12DynamicRHI();
@@ -528,10 +528,10 @@ public:
 	virtual bool HandleSpecialLock(FRHILockTextureResult& Result, FD3D12Texture* InTexture, const FRHILockTextureArgs& Arguments) { return false; }
 	virtual bool HandleSpecialUnlock(FRHICommandListBase& RHICmdList, FD3D12Texture* InTexture, const FRHILockTextureArgs& Arguments) { return false; }
 
-	FD3D12Adapter& GetAdapter(uint32_t Index = 0) { return *ChosenAdapters[Index]; }
-	const FD3D12Adapter& GetAdapter(uint32_t Index = 0) const { return *ChosenAdapters[Index]; }
+	FD3D12Adapter& GetAdapter() { return *ChosenAdapter; }
+	const FD3D12Adapter& GetAdapter() const { return *ChosenAdapter; }
 
-	uint32 GetNumAdapters() const { return ChosenAdapters.Num(); }
+	bool HasAdapter() const { return ChosenAdapter.IsValid(); }
 
 	bool IsPixEventEnabled() const { return bPixEventEnabled; }
 
@@ -546,7 +546,7 @@ public:
 
 protected:
 
-	TArray<TSharedPtr<FD3D12Adapter>> ChosenAdapters;
+	TSharedPtr<FD3D12Adapter> ChosenAdapter;
 
 #if D3D12RHI_SUPPORTS_WIN_PIX
 	void* WinPixGpuCapturerHandle = nullptr;
@@ -705,7 +705,7 @@ private:
 	void* WinPixGpuCapturerHandle = nullptr;
 #endif
 
-	TArray<TSharedPtr<FD3D12Adapter>> ChosenAdapters;
+	TSharedPtr<FD3D12Adapter> ChosenAdapter;
 
 	// set MaxSupportedFeatureLevel and ChosenAdapter
 	void FindAdapter();
diff --git a/Engine/Source/Runtime/D3D12RHI/Private/D3D12Util.cpp b/Engine/Source/Runtime/D3D12RHI/Private/D3D12Util.cpp
--- a/Engine/Source/Runtime/D3D12RHI/Private/D3D12Util.cpp
+++ b/Engine/Source/Runtime/D3D12RHI/Private/D3D12Util.cpp
@@ -30,12 +30,11 @@ D3D12Util.h: D3D RHI utility implementation.
 template<typename PerDeviceFunction>
 void FD3D12DynamicRHI::ForEachDevice(ID3D12Device* inDevice, const PerDeviceFunction& pfPerDeviceFunction)
 {
-	for (uint32 AdapterIndex = 0; AdapterIndex < GetNumAdapters(); ++AdapterIndex)
+	if (FD3D12Adapter* const D3D12Adapter = ChosenAdapter.Get())
 	{
-		FD3D12Adapter& D3D12Adapter = GetAdapter(AdapterIndex);
 		for (uint32 GPUIndex : FRHIGPUMask::All())
 		{
-			FD3D12Device* D3D12Device = D3D12Adapter.GetDevice(GPUIndex);
+			FD3D12Device* D3D12Device = D3D12Adapter->GetDevice(GPUIndex);
 			if (inDevice == nullptr || D3D12Device->GetDevice() == inDevice)
 			{
 				pfPerDeviceFunction(D3D12Device);
@@ -765,11 +764,11 @@ void FD3D12DynamicRHI::TerminateOnOutOfMemory(ID3D12Device* InDevice, HRESULT D3
 		}
 	});
 
-	// if InDevice == nullptr, just pick the first available adapter
-	if (!Adapter && GetNumAdapters() == 1)
+	// if InDevice == nullptr, just pick the chosen adapter
+	if (!Adapter && ChosenAdapter.IsValid())
 	{
 		check(!InDevice);
-		Adapter = &GetAdapter(0);
+		Adapter = ChosenAdapter.Get();
 	}
 
 	if (Adapter)
diff --git a/Engine/Source/Runtime/D3D12RHI/Private/Windows/WindowsD3D12Device.cpp b/Engine/Source/Runtime/D3D12RHI/Private/Windows/WindowsD3D12Device.cpp
--- a/Engine/Source/Runtime/D3D12RHI/Private/Windows/WindowsD3D12Device.cpp
+++ b/Engine/Source/Runtime/D3D12RHI/Private/Windows/WindowsD3D12Device.cpp
@@ -551,7 +551,7 @@ static bool SafeTestD3D12CreateDevice(IDXGIAdapter* Adapter, D3D_FEATURE_LEVEL M
 static bool SupportsDepthBoundsTest(FD3D12DynamicRHI* D3DRHI)
 {
 	// Determines if the primary adapter supports depth bounds test
-	check(D3DRHI && D3DRHI->GetNumAdapters() >= 1);
+	check(D3DRHI && D3DRHI->HasAdapter());
 
 	return D3DRHI->GetAdapter().IsDepthBoundsTestSupported();
 }
@@ -568,21 +568,13 @@ static UE::Color::FColorSpace MakeColorSpaceFromDXGI(const DXGI_OUTPUT_DESC1& Ou
 bool FD3D12DynamicRHI::SetupDisplayHDRMetaData()
 {
 	// Determines if any displays support HDR
-	check(GetNumAdapters() >= 1);
+	check(HasAdapter());
 
 	TArray<TRefCountPtr<IDXGIAdapter> > DXGIAdapters;
-	bool bStaleAdapters = false;
 
-	const int32 NumAdapters = GetNumAdapters();
-	for (int32 AdapterIndex = 0; AdapterIndex < NumAdapters; ++AdapterIndex)
-	{
-		FD3D12Adapter& CurrentAdapter = GetAdapter(AdapterIndex);
-		DXGIAdapters.Add(GetAdapter(AdapterIndex).GetAdapter());
-		if (CurrentAdapter.GetDXGIFactory2() != nullptr && !CurrentAdapter.GetDXGIFactory2()->IsCurrent())
-		{
-			bStaleAdapters = true;
-		}
-	}
+	FD3D12Adapter& CurrentAdapter = *ChosenAdapter.Get();
+	DXGIAdapters.Add(CurrentAdapter.GetAdapter());
+	const bool bStaleAdapters = CurrentAdapter.GetDXGIFactory2() && !CurrentAdapter.GetDXGIFactory2()->IsCurrent();
 
 #if PLATFORM_WINDOWS
 	// if we found that the list of adapters is stale (changed windows HDR setting), try to update it with the new list
@@ -602,7 +594,6 @@ bool FD3D12DynamicRHI::SetupDisplayHDRMetaData()
 				DXGIAdapters.Add(TempAdapter);
 			}
 		}
-
 	}
 #endif
 
@@ -743,18 +734,18 @@ bool FD3D12DynamicRHIModule::IsSupported(ERHIFeatureLevel::Type RequestedFeature
 	}
 
 	// If not computed yet
-	if (ChosenAdapters.Num() == 0)
+	if (!ChosenAdapter.IsValid())
 	{
 		FindAdapter();
 	}
 
-	if (ChosenAdapters.Num() == 0)
+	if (!ChosenAdapter.IsValid())
 	{
 		UE_LOG(LogD3D12RHI, Log, TEXT("No adapters were found."));
 		return false;
 	}
 
-	const FD3D12Adapter* Adapter = ChosenAdapters[0].Get();
+	const FD3D12Adapter* const Adapter = ChosenAdapter.Get();
 	if (!Adapter || !Adapter->GetDesc().IsValid())
 	{
 		UE_LOG(LogD3D12RHI, Log, TEXT("Adapter was not found"));
@@ -834,7 +825,7 @@ static uint32 CountAdapterOutputs(TRefCountPtr<IDXGIAdapter>& Adapter)
 void FD3D12DynamicRHIModule::FindAdapter()
 {
 	// Once we've chosen one we don't need to do it again.
-	check(ChosenAdapters.Num() == 0);
+	check(!ChosenAdapter.IsValid());
 
 #if !UE_BUILD_SHIPPING
 	if (CVarExperimentalShaderModels.GetValueOnAnyThread() == 1)
@@ -1038,46 +1029,40 @@ void FD3D12DynamicRHIModule::FindAdapter()
 		}
 	}
 
-	TSharedPtr<FD3D12Adapter> NewAdapter;
 	if (bFavorDiscreteAdapter)
 	{
 		// We assume Intel is integrated graphics (slower than discrete) than NVIDIA or AMD cards and rather take a different one
 		if (PreferredAdapter.IsValid())
 		{
-			NewAdapter = TSharedPtr<FD3D12Adapter>(new FWindowsD3D12Adapter(PreferredAdapter));
-			ChosenAdapters.Add(NewAdapter);
+			ChosenAdapter = MakeShareable(new FWindowsD3D12Adapter(PreferredAdapter));
 		}
 		else if (BestMemoryAdapter.IsValid())
 		{
-			NewAdapter = TSharedPtr<FD3D12Adapter>(new FWindowsD3D12Adapter(BestMemoryAdapter));
-			ChosenAdapters.Add(NewAdapter);
+			ChosenAdapter = MakeShareable(new FWindowsD3D12Adapter(BestMemoryAdapter));
 		}
 		else if (FirstDiscreteAdapter.IsValid())
 		{
-			NewAdapter = TSharedPtr<FD3D12Adapter>(new FWindowsD3D12Adapter(FirstDiscreteAdapter));
-			ChosenAdapters.Add(NewAdapter);
+			ChosenAdapter = MakeShareable(new FWindowsD3D12Adapter(FirstDiscreteAdapter));
 		}
 		else
 		{
-			NewAdapter = TSharedPtr<FD3D12Adapter>(new FWindowsD3D12Adapter(FirstAdapter));
-			ChosenAdapters.Add(NewAdapter);
+			ChosenAdapter = MakeShareable(new FWindowsD3D12Adapter(FirstAdapter));
 		}
 	}
 	else
 	{
-		NewAdapter = TSharedPtr<FD3D12Adapter>(new FWindowsD3D12Adapter(FirstAdapter));
-		ChosenAdapters.Add(NewAdapter);
+		ChosenAdapter = MakeShareable(new FWindowsD3D12Adapter(FirstAdapter));
 	}
 
 #if D3D12_CORE_ENABLED
 	UE_LOG(LogD3D12RHI, Log, TEXT("DirectX Agility SDK runtime %s."), CheckIfAgilitySDKLoaded() ? TEXT("found") : TEXT("not found"));
 #endif
 
-	if (ChosenAdapters.Num() > 0 && ChosenAdapters[0]->GetDesc().IsValid())
+	if (ChosenAdapter.IsValid() && ChosenAdapter->GetDesc().IsValid())
 	{
-		UE_LOG(LogD3D12RHI, Log, TEXT("Chosen D3D12 Adapter Id = %u"), ChosenAdapters[0]->GetAdapterIndex());
+		UE_LOG(LogD3D12RHI, Log, TEXT("Chosen D3D12 Adapter Id = %u"), ChosenAdapter->GetAdapterIndex());
 
-		const DXGI_ADAPTER_DESC& AdapterDesc = ChosenAdapters[0]->GetD3DAdapterDesc();
+		const DXGI_ADAPTER_DESC& AdapterDesc = ChosenAdapter->GetD3DAdapterDesc();
 		GRHIVendorId = AdapterDesc.VendorId;
 		GRHIAdapterName = AdapterDesc.Description;
 		GRHIDeviceId = AdapterDesc.DeviceId;
@@ -1089,19 +1074,6 @@ void FD3D12DynamicRHIModule::FindAdapter()
 	}
 }
 
-static bool DoesAnyAdapterSupportSM6(const TArray<TSharedPtr<FD3D12Adapter>>& Adapters)
-{
-	for (const TSharedPtr<FD3D12Adapter>& Adapter : Adapters)
-	{
-		if (IsAdapterSupported(Adapter.Get(), ERHIFeatureLevel::SM6))
-		{
-			return true;
-		}
-	}
-
-	return false;
-}
-
 FDynamicRHI* FD3D12DynamicRHIModule::CreateRHI(ERHIFeatureLevel::Type RequestedFeatureLevel)
 {
 	GShaderPlatformForFeatureLevel[ERHIFeatureLevel::ES2_REMOVED] = SP_NumPlatforms;
@@ -1109,7 +1081,7 @@ FDynamicRHI* FD3D12DynamicRHIModule::CreateRHI(ERHIFeatureLevel::Type RequestedF
 	GShaderPlatformForFeatureLevel[ERHIFeatureLevel::SM4_REMOVED] = SP_NumPlatforms;
 	GShaderPlatformForFeatureLevel[ERHIFeatureLevel::SM5] = SP_PCD3D_SM5;
 	GShaderPlatformForFeatureLevel[ERHIFeatureLevel::SM6] = SP_NumPlatforms;
-	if (DoesAnyAdapterSupportSM6(ChosenAdapters))
+	if (IsAdapterSupported(ChosenAdapter.Get(), ERHIFeatureLevel::SM6))
 	{
 		GShaderPlatformForFeatureLevel[ERHIFeatureLevel::SM6] = SP_PCD3D_SM6;
 	}
@@ -1141,21 +1113,22 @@ FDynamicRHI* FD3D12DynamicRHIModule::CreateRHI(ERHIFeatureLevel::Type RequestedF
 	bool bPixEventEnabled = false;
 #endif // USE_PIX
 	
-	if (ChosenAdapters.Num() > 0 && ChosenAdapters[0].IsValid())
+	if (ChosenAdapter.IsValid())
 	{
-		const bool bIsIntegrated = ChosenAdapters[0].Get()->GetDesc().bUMA;
+		const FD3D12AdapterDesc& Desc = ChosenAdapter->GetDesc();
+		const bool bIsIntegrated = Desc.bUMA;
 		FGenericCrashContext::SetEngineData(TEXT("RHI.IntegratedGPU"), bIsIntegrated ? TEXT("true") : TEXT("false"));
 		GRHIDeviceIsIntegrated = bIsIntegrated;
-		GRHIDeviceIsUMASeparateVA = ChosenAdapters[0].Get()->GetDesc().bUMA;
-		GRHIDeviceIsCacheCoherentUMA = ChosenAdapters[0].Get()->GetDesc().bCacheCoherentUMA;
-		GRHIDeviceSupportsGPUUploadHeaps = ChosenAdapters[0].Get()->GetDesc().bSupportsGPUUploadHeaps;
+		GRHIDeviceIsUMASeparateVA = Desc.bUMA;
+		GRHIDeviceIsCacheCoherentUMA = Desc.bCacheCoherentUMA;
+		GRHIDeviceSupportsGPUUploadHeaps = Desc.bSupportsGPUUploadHeaps;
 		UE_LOG(LogD3D12RHI, Log, TEXT("Integrated GPU (iGPU): %s"), GRHIDeviceIsIntegrated ? TEXT("true") : TEXT("false"));
 	}
 
 	const FString FeatureLevelString = LexToString(GMaxRHIFeatureLevel);
 	UE_LOG(LogD3D12RHI, Display, TEXT("Creating D3D12 RHI with Max Feature Level %s"), *FeatureLevelString);
 
-	GD3D12RHI = new FD3D12DynamicRHI(ChosenAdapters, bPixEventEnabled);
+	GD3D12RHI = new FD3D12DynamicRHI(ChosenAdapter, bPixEventEnabled);
 #if ENABLE_RHI_VALIDATION
 	if (FParse::Param(FCommandLine::Get(), TEXT("RHIValidation")))
 	{
@@ -1452,9 +1425,9 @@ void FD3D12DynamicRHI::Init()
 	}
 #endif
 
-	// Create a device chain for each of the adapters we have chosen. This could be a single discrete card,
+	// Create a device chain for the adapter we have chosen. This could be a single discrete card,
 	// a set discrete cards linked together (i.e. SLI/Crossfire) an Integrated device or any combination of the above
-	for (TSharedPtr<FD3D12Adapter>& Adapter : ChosenAdapters)
+	if (FD3D12Adapter* const Adapter = ChosenAdapter.Get())
 	{
 		check(Adapter->GetDesc().IsValid());
 		Adapter->InitializeDevices();
@@ -2118,12 +2091,10 @@ static void GetAvailableResolutionsForOutput(IDXGIOutput* Output, FScreenResolut
 bool FD3D12DynamicRHI::RHIGetAvailableResolutions(FScreenResolutionArray& Resolutions, bool bIgnoreRefreshRate)
 {
 	// TODO: The resolution information can be stale if the adapter's factory isn't current. Use RHIGetAvailableResolutionsForDisplay instead.
-	FD3D12Adapter& ChosenAdapter = GetAdapter();
-
 	HRESULT HResult = S_OK;
 	TRefCountPtr<IDXGIAdapter> Adapter;
 	//TODO: should only be called on display out device
-	HResult = ChosenAdapter.EnumAdapters(Adapter.GetInitReference());
+	HResult = ChosenAdapter->EnumAdapters(Adapter.GetInitReference());
 
 	if (DXGI_ERROR_NOT_FOUND == HResult)
 	{
@@ -2191,8 +2162,7 @@ bool FD3D12DynamicRHI::RHIGetAvailableResolutionsForDisplay(FScreenResolutionArr
 				// Undocumented but assume exclusive only works if the the display is primary, i.e. top-left at (0, 0), or the display is connected to the chosen adapter.
 				if (OutputDesc.DesktopCoordinates.left != 0 || OutputDesc.DesktopCoordinates.top != 0)
 				{
-					const FD3D12Adapter* const ChosenAdapter = ChosenAdapters[0].Get();
-					if (ChosenAdapter)
+					if (ChosenAdapter.IsValid())
 					{
 						const FD3D12AdapterDesc& ChosenDesc = ChosenAdapter->GetDesc();
 						if (ChosenDesc.IsValid())
