[RT Validation] Nvidia raytracing validation which is useful for debugging gpu crashes in internal RT structures / shaders and raytracing perf.
Note: this has a non-trivial performance cost. Should only be enabled during development / debugging.
How to enable:
Cvar: r.D3D12.DXR.RaytracingValidation 1 OR cmd line: -validateRayTracing
AND system env variable: NV_ALLOW_RAYTRACING_VALIDATION=1 (will fail on startup if this is not set.)

#rnx

#Synced-CL 49975040

diff --git a/Engine/Source/Runtime/D3D12RHI/Private/D3D12RHI.cpp b/Engine/Source/Runtime/D3D12RHI/Private/D3D12RHI.cpp
--- a/Engine/Source/Runtime/D3D12RHI/Private/D3D12RHI.cpp
+++ b/Engine/Source/Runtime/D3D12RHI/Private/D3D12RHI.cpp
@@ -300,6 +300,15 @@ void FD3D12DynamicRHI::Shutdown()
 	// Reset the RHI initialized flag.
 	GIsRHIInitialized = false;
 
+#if WITH_NVAPI_RAYTRACING_VALIDATION
+	if (RayTracingValidationCallbackHandle)
+	{		
+		const NvAPI_Status NvStatus = NvAPI_D3D12_UnregisterRaytracingValidationMessageCallback(GetAdapter().GetD3DDevice5(), RayTracingValidationCallbackHandle);
+		check(NvStatus == NVAPI_OK);
+		RayTracingValidationCallbackHandle = nullptr;
+	}
+#endif
+
 #if WITH_AMD_AGS
 	if (AmdAgsContext)
 	{
@@ -549,6 +558,17 @@ void FD3D12DynamicRHI::RHIEndFrame_RenderThread(FRHICommandListImmediate& RHICmd
 	});
 }
 
+#if WITH_NVAPI_RAYTRACING_VALIDATION
+void FD3D12DynamicRHI::FlushRaytracingValidationMessages()
+{
+	if (RayTracingValidationCallbackHandle)
+	{
+		const NvAPI_Status NvStatus = NvAPI_D3D12_FlushRaytracingValidationMessages(GetAdapter().GetD3DDevice5());
+		ensure(NvStatus == NVAPI_OK);
+	}
+}
+#endif
+
 void FD3D12DynamicRHI::RHIEndFrame(const FRHIEndFrameArgs& Args)
 {
 	if (FD3D12Adapter* const Adapter = ChosenAdapter.Get())
@@ -604,6 +624,13 @@ void FD3D12DynamicRHI::RHIEndFrame(const FRHIEndFrameArgs& Args)
 		Payload.EndFrameEvent = UE::RHI::GPUProfiler::FEvent::FFrameBoundary(Pipeline, Args, 0);
 	});
 
+#if WITH_NVAPI_RAYTRACING_VALIDATION
+	if (RayTracingValidationCallbackHandle)
+	{
+		EnqueueEndOfPipeTask([this]() { FlushRaytracingValidationMessages(); });
+	}
+#endif
+
 	// Pump the interrupt queue to gather completed events
 	// (required if we're not using an interrupt thread).
 	ProcessInterruptQueueUntil(nullptr);
diff --git a/Engine/Source/Runtime/D3D12RHI/Private/D3D12RHIDefinitions.h b/Engine/Source/Runtime/D3D12RHI/Private/D3D12RHIDefinitions.h
--- a/Engine/Source/Runtime/D3D12RHI/Private/D3D12RHIDefinitions.h
+++ b/Engine/Source/Runtime/D3D12RHI/Private/D3D12RHIDefinitions.h
@@ -55,6 +55,12 @@
 	#define D3D12RHI_SUPPORTS_UNCOMPRESSED_UAV 0
 #endif
 
+#if WITH_NVAPI && (D3D12_MAX_DEVICE_INTERFACE >= 5) && D3D12_RHI_RAYTRACING && !UE_BUILD_SHIPPING
+	#define WITH_NVAPI_RAYTRACING_VALIDATION 1
+#else
+	#define WITH_NVAPI_RAYTRACING_VALIDATION 0
+#endif
+
 // This value controls how many root constant buffers can be used per shader stage in a root signature.
 // Note: Using root descriptors significantly increases the size of root signatures (each root descriptor is 2 DWORDs).
 #define MAX_ROOT_CBVS	MAX_CBS
diff --git a/Engine/Source/Runtime/D3D12RHI/Private/D3D12RHIPrivate.h b/Engine/Source/Runtime/D3D12RHI/Private/D3D12RHIPrivate.h
--- a/Engine/Source/Runtime/D3D12RHI/Private/D3D12RHIPrivate.h
+++ b/Engine/Source/Runtime/D3D12RHI/Private/D3D12RHIPrivate.h
@@ -237,6 +237,10 @@ private:
 
 	virtual void UpdateMemoryStats();
 
+#if WITH_NVAPI_RAYTRACING_VALIDATION
+	void FlushRaytracingValidationMessages();
+#endif
+
 public:
 	template <typename ...Args>
 	void DeferredDelete(Args&&... InArgs)
@@ -568,6 +572,10 @@ protected:
 
 	INTCExtensionContext* IntelExtensionContext = nullptr;
 
+#if WITH_NVAPI_RAYTRACING_VALIDATION
+	void* RayTracingValidationCallbackHandle = nullptr;
+#endif
+
 	// See r.PSOPrecache.D3D12.DriverCacheAware. The CVar is cached
 	// to prevent hotfixes or config dynamic layers from altering the
 	// value at runtime, which is not supported.
diff --git a/Engine/Source/Runtime/D3D12RHI/Private/D3D12Util.cpp b/Engine/Source/Runtime/D3D12RHI/Private/D3D12Util.cpp
--- a/Engine/Source/Runtime/D3D12RHI/Private/D3D12Util.cpp
+++ b/Engine/Source/Runtime/D3D12RHI/Private/D3D12Util.cpp
@@ -981,6 +981,10 @@ void FD3D12DynamicRHI::OutputGPUCrashReport(FTextBuilder& ErrorMessage)
 
 void FD3D12DynamicRHI::TerminateOnGPUCrash()
 {
+#if WITH_NVAPI_RAYTRACING_VALIDATION
+	FlushRaytracingValidationMessages();
+#endif
+
 	FTextBuilder ErrorMessage;
 	OutputGPUCrashReport(ErrorMessage);
 
diff --git a/Engine/Source/Runtime/D3D12RHI/Private/Windows/WindowsD3D12Device.cpp b/Engine/Source/Runtime/D3D12RHI/Private/Windows/WindowsD3D12Device.cpp
--- a/Engine/Source/Runtime/D3D12RHI/Private/Windows/WindowsD3D12Device.cpp
+++ b/Engine/Source/Runtime/D3D12RHI/Private/Windows/WindowsD3D12Device.cpp
@@ -41,6 +41,21 @@ bool GDX12INTCAtomicUInt64Emulation = false;
 
 FD3D12DynamicRHI* GD3D12RHI = nullptr;
 
+#if WITH_NVAPI
+int32 GEnableNVAPIRaytracingValidation = 0;
+static FAutoConsoleVariableRef CVarNVAPIRaytracingValidation(
+	TEXT("r.D3D12.DXR.RaytracingValidation"),
+	GEnableNVAPIRaytracingValidation,
+	TEXT("Enables NVAPI Raytracing validation."),
+	ECVF_ReadOnly
+);
+static TAutoConsoleVariable<FString> CVarNVAPIRaytracingFilterMessages(
+	TEXT("r.D3D12.DXR.RaytracingValidation.IgnoreList"),
+	"",
+	TEXT("List of warnings or errors to ignore to prevent spam from NVAPI Raytracing validation. For instance: warning1;warning2 etc"),
+	ECVF_ReadOnly);
+#endif
+
 int32 GMinimumWindowsBuildVersionForRayTracing = 0;
 static FAutoConsoleVariableRef CVarMinBuildVersionForRayTracing(
 	TEXT("r.D3D12.DXR.MinimumWindowsBuildVersion"),
@@ -1550,6 +1565,75 @@ void FD3D12DynamicRHI::Init()
 			// NvAPI_D3D12_GetRaytracingCaps appears to leak about 100 MB so avoid calling if RayTracing is impossible
 			if (GMaxRHIFeatureLevel >= ERHIFeatureLevel::SM6)
 			{
+#if WITH_NVAPI_RAYTRACING_VALIDATION
+				{
+					ID3D12Device5* Device5 = GetAdapter().GetD3DDevice5();
+					if (GRHISupportsRayTracing
+						&& Device5 != nullptr
+						&& (GEnableNVAPIRaytracingValidation || FParse::Param(FCommandLine::Get(), TEXT("validateRayTracing"))))
+					{
+						const NvAPI_Status NvStatusVal = NvAPI_D3D12_EnableRaytracingValidation(Device5, NVAPI_D3D12_RAYTRACING_VALIDATION_FLAG_NONE);
+						if (NvStatusVal == NVAPI_OK)
+						{
+							auto RtValidationMessageCB = [](void* pUserData, NVAPI_D3D12_RAYTRACING_VALIDATION_MESSAGE_SEVERITY Severity, const char* MessageCode, const char* Message, const char* MessageDetails)
+								{
+									TArray<FString>* RtValidationCodesToIgnoreList = (TArray<FString>*)pUserData;
+
+									if (RtValidationCodesToIgnoreList->Contains(MessageCode))
+									{
+										return;
+									}
+
+									switch (Severity)
+									{
+									case NVAPI_D3D12_RAYTRACING_VALIDATION_MESSAGE_SEVERITY_ERROR:
+									{
+										UE_LOG(LogD3D12RHI, Error, TEXT("Ray Tracing Validation message: [%hs] %hs\n%hs"), MessageCode, Message, MessageDetails);
+										break;
+									}
+									case NVAPI_D3D12_RAYTRACING_VALIDATION_MESSAGE_SEVERITY_WARNING:
+									{
+										UE_LOG(LogD3D12RHI, Warning, TEXT("Ray Tracing Validation message: [%hs] %hs\n%hs"), MessageCode, Message, MessageDetails);
+										break;
+									}
+									default:
+									{
+										UE_LOG(LogD3D12RHI, Warning, TEXT("Ray Tracing Validation message: [%hs] %hs\n%hs"), MessageCode, Message, MessageDetails);
+										break;
+									}
+									}
+								};
+
+							static TArray<FString> RtValidationCodesToIgnoreList = []() -> TArray<FString> {
+								TArray<FString> Res;
+								FString CodesToIgnoreStr = CVarNVAPIRaytracingFilterMessages.GetValueOnRenderThread();
+								CodesToIgnoreStr.ParseIntoArray(Res, TEXT(";"));
+								return Res;
+								}();
+
+							const NvAPI_Status NvStatusCb = NvAPI_D3D12_RegisterRaytracingValidationMessageCallback(Device5, RtValidationMessageCB, &RtValidationCodesToIgnoreList, &RayTracingValidationCallbackHandle);
+							checkf(NvStatusCb == NVAPI_OK, TEXT("NvAPI Raytracing Validation Enabled, but NvAPI_D3D12_RegisterRaytracingValidationMessageCallback failed with %d."), (uint32)NvStatusCb);
+							if (NvStatusCb == NVAPI_OK)
+							{
+								UE_LOG(LogD3D12RHI, Log, TEXT("NvAPI Raytracing Validation Enabled (-validateRayTracing)"));
+							}
+						}
+						else if (NvStatusVal == NVAPI_ACCESS_DENIED)
+						{
+							checkf(false, TEXT("Failed to enable NvAPI Raytracing Validation. The env variable NV_ALLOW_RAYTRACING_VALIDATION=1 must be set on the system."));
+						}
+						else if (NvStatusVal == NVAPI_INVALID_CALL)
+						{
+							checkf(false, TEXT("Failed to enable NvAPI Raytracing Validation (NVAPI_INVALID_CALL). Raytracing Validation must be initialized before any RT-Related calls."));
+						}
+						else
+						{
+							checkf(false, TEXT("Failed to enable NvAPI Raytracing Validation. Error=%d"), (uint32)NvStatusVal);
+						}
+					}
+				}
+#endif // WITH_NVAPI_RAYTRACING_VALIDATION
+
 				NvAPI_Status NvStatusSER = NvAPI_D3D12_IsNvShaderExtnOpCodeSupported(GetAdapter().GetD3DDevice(), NV_EXTN_OP_HIT_OBJECT_REORDER_THREAD, &GRHIGlobals.SupportsShaderExecutionReordering);
 				if (NvStatusSER == NVAPI_OK && GRHIGlobals.SupportsShaderExecutionReordering)
 				{
