Vendor Extensions
- Adding Nvidia to QBSS so that more shader types can use NVAPI extensions.
- Adding validation of vendor extensions to D3D12 shaders so that we don't precache shaders meant for other vendors.

#Synced-CL 50075005

diff --git a/Engine/Source/Runtime/D3D12RHI/Private/D3D12RootSignature.cpp b/Engine/Source/Runtime/D3D12RHI/Private/D3D12RootSignature.cpp
--- a/Engine/Source/Runtime/D3D12RHI/Private/D3D12RootSignature.cpp
+++ b/Engine/Source/Runtime/D3D12RHI/Private/D3D12RootSignature.cpp
@@ -437,7 +437,7 @@ FD3D12RootSignatureDesc::FD3D12RootSignatureDesc(const FD3D12QuantizedBoundShade
 	}
 
 #if WITH_NVAPI
-	if (QBSS.RootSignatureType == RS_RayTracingGlobal && IsRHIDeviceNVIDIA())
+	if (QBSS.bNeedsNvIntrinsicsSpace)
 	{
 		check(RootParameterCount < MaxRootParameters);
 
diff --git a/Engine/Source/Runtime/D3D12RHI/Private/D3D12Shaders.cpp b/Engine/Source/Runtime/D3D12RHI/Private/D3D12Shaders.cpp
--- a/Engine/Source/Runtime/D3D12RHI/Private/D3D12Shaders.cpp
+++ b/Engine/Source/Runtime/D3D12RHI/Private/D3D12Shaders.cpp
@@ -134,6 +134,19 @@ static bool ValidateShaderIsUsable(FD3D12ShaderData* InShader, EShaderFrequency
 	}
 #endif
 
+#if D3D12RHI_NEEDS_VENDOR_EXTENSIONS
+	// At the time that this code was added, we only have AMD, Nvidia and Intel extensions and GetRHIDeviceVendorId handles all these vendors.
+	const EGpuVendorId CurrentVendorId = GetRHIDeviceVendorId();
+
+	for (const FShaderCodeVendorExtension& Extension : InShader->VendorExtensions)
+	{
+		if (Extension.VendorId != CurrentVendorId)
+		{
+			return false;
+		}
+	}
+#endif
+
 	return true;
 }
 
diff --git a/Engine/Source/Runtime/D3D12RHI/Private/D3D12Util.cpp b/Engine/Source/Runtime/D3D12RHI/Private/D3D12Util.cpp
--- a/Engine/Source/Runtime/D3D12RHI/Private/D3D12Util.cpp
+++ b/Engine/Source/Runtime/D3D12RHI/Private/D3D12Util.cpp
@@ -1196,6 +1196,22 @@ bool NeedsAgsIntrinsicsSpace(const FD3D12ShaderData& ShaderData)
 	return false;
 }
 
+bool NeedsNvIntrinsicsSpace(const FD3D12ShaderData& ShaderData)
+{
+#if D3D12RHI_NEEDS_VENDOR_EXTENSIONS
+	for (const FShaderCodeVendorExtension& Extension : ShaderData.VendorExtensions)
+	{
+		if (Extension.VendorId == EGpuVendorId::Nvidia)
+		{
+			// https://developer.nvidia.com/unlocking-gpu-intrinsics-hlsl
+			return true;
+		}
+	}
+#endif
+
+	return false;
+}
+
 static void SetBoundShaderStateFlags(FD3D12QuantizedBoundShaderState& OutQBSS, const FD3D12ShaderData* ShaderData)
 {
 	if (ShaderData)
@@ -1224,6 +1240,7 @@ static void QuantizeBoundShaderStateCommon(
 	{
 		FD3D12QuantizedBoundShaderState::InitShaderRegisterCounts(ResourceBindingTier, ShaderData->ResourceCounts, OutQBSS.RegisterCounts[ShaderVisibility], bAllowUAVs);
 		OutQBSS.bNeedsAgsIntrinsicsSpace |= NeedsAgsIntrinsicsSpace(*ShaderData);
+		OutQBSS.bNeedsNvIntrinsicsSpace |= NeedsNvIntrinsicsSpace(*ShaderData);
 	}
 
 	SetBoundShaderStateFlags(OutQBSS, ShaderData);
diff --git a/Engine/Source/Runtime/D3D12RHI/Private/D3D12Util.h b/Engine/Source/Runtime/D3D12RHI/Private/D3D12Util.h
--- a/Engine/Source/Runtime/D3D12RHI/Private/D3D12Util.h
+++ b/Engine/Source/Runtime/D3D12RHI/Private/D3D12Util.h
@@ -125,11 +125,12 @@ struct FD3D12QuantizedBoundShaderState
 		{
 			uint8 bAllowIAInputLayout : 1;
 			uint8 bNeedsAgsIntrinsicsSpace : 1;
+			uint8 bNeedsNvIntrinsicsSpace : 1;
 			uint8 bUseDiagnosticBuffer : 1;
 			uint8 bUseDirectlyIndexedResourceHeap : 1;
 			uint8 bUseDirectlyIndexedSamplerHeap : 1;
 			uint8 bUseRootConstants : 1;
-			uint8 Padding : 2;
+			uint8 Padding : 1;
 		};
 		uint8 Flags;
 	};
diff --git a/Engine/Source/Runtime/RHI/Private/RHI.cpp b/Engine/Source/Runtime/RHI/Private/RHI.cpp
--- a/Engine/Source/Runtime/RHI/Private/RHI.cpp
+++ b/Engine/Source/Runtime/RHI/Private/RHI.cpp
@@ -1314,6 +1314,32 @@ bool IsRHIDeviceApple()
     return GRHIVendorId == (uint32) EGpuVendorId::Apple;
 }
 
+EGpuVendorId GetRHIDeviceVendorId()
+{
+	if (IsRHIDeviceAMD())
+	{
+		return EGpuVendorId::Amd;
+	}
+	if (IsRHIDeviceIntel())
+	{
+		return EGpuVendorId::Intel;
+	}
+	if (IsRHIDeviceNVIDIA())
+	{
+		return EGpuVendorId::Nvidia;
+	}
+	if (IsRHIDeviceQualcomm())
+	{
+		return EGpuVendorId::Qualcomm;
+	}
+	if (IsRHIDeviceApple())
+	{
+		return EGpuVendorId::Apple;
+	}
+
+	return EGpuVendorId::Unknown;
+}
+
 uint32 RHIGetMetalShaderLanguageVersion(const FStaticShaderPlatform Platform)
 {
     if (IsMetalPlatform(Platform))
diff --git a/Engine/Source/Runtime/RHI/Public/RHI.h b/Engine/Source/Runtime/RHI/Public/RHI.h
--- a/Engine/Source/Runtime/RHI/Public/RHI.h
+++ b/Engine/Source/Runtime/RHI/Public/RHI.h
@@ -45,6 +45,8 @@ RHI_API bool IsRHIDeviceQualcomm();
 // to trigger GPU specific optimizations and fallbacks
 RHI_API bool IsRHIDeviceApple();
 
+RHI_API EGpuVendorId GetRHIDeviceVendorId();
+
 namespace UE::RHI
 {
 	// Returns true when GPU crash debugging is enabled.
