Fix incorrect registerSpace for NVAPI extensions.

#Synced-CL 50064624

diff --git a/Engine/Source/Runtime/D3D12RHI/Private/Windows/WindowsD3D12PipelineState.cpp b/Engine/Source/Runtime/D3D12RHI/Private/Windows/WindowsD3D12PipelineState.cpp
--- a/Engine/Source/Runtime/D3D12RHI/Private/Windows/WindowsD3D12PipelineState.cpp
+++ b/Engine/Source/Runtime/D3D12RHI/Private/Windows/WindowsD3D12PipelineState.cpp
@@ -877,7 +877,7 @@ static FORCEINLINE NVAPI_D3D12_PSO_SET_SHADER_EXTENSION_SLOT_DESC GetNVShaderExt
 	// TODO: Temp workaround until RenderDoc fixes SM 5.0 opcode detection (need to ignore register space)
 	ShdExtensionDesc.registerSpace = 0xFFFFFFFF;
 #else
-	ShdExtensionDesc.registerSpace = 0; // TODO: Should use the same special space as AMD to avoid driver pattern matching
+	ShdExtensionDesc.registerSpace = UE_HLSL_SPACE_NV_SHADER_EXTN;
 #endif
 	return ShdExtensionDesc;
 }
