[HWRT] Fix issues due to NvIntrinsicsSpace not being fully bound to ray tracing global root signature.

#Synced-CL 50064624

diff --git a/Engine/Source/Runtime/D3D12RHI/Private/D3D12Util.cpp b/Engine/Source/Runtime/D3D12RHI/Private/D3D12Util.cpp
--- a/Engine/Source/Runtime/D3D12RHI/Private/D3D12Util.cpp
+++ b/Engine/Source/Runtime/D3D12RHI/Private/D3D12Util.cpp
@@ -1586,6 +1586,7 @@ const FD3D12RootSignature* FD3D12Adapter::GetGlobalRayTracingRootSignature(const
 	QBSS.ShaderBindingLayout = ShaderBindingLayout;
 	QBSS.RootSignatureType = RS_RayTracingGlobal;
 	QBSS.bUseDiagnosticBuffer = true;
+	QBSS.bNeedsNvIntrinsicsSpace = IsRHIDeviceNVIDIA();
 
 #if PLATFORM_SUPPORTS_BINDLESS_RENDERING
 	QBSS.bUseDirectlyIndexedResourceHeap = GetBindlessDescriptorAllocator().AreResourcesBindless();
