Adding Mesh/Amplification shader UAV support to D3D12RHI.

#Synced-CL 50167313

diff --git a/Engine/Source/Runtime/D3D12RHI/Private/D3D12CommandContext.h b/Engine/Source/Runtime/D3D12RHI/Private/D3D12CommandContext.h
--- a/Engine/Source/Runtime/D3D12RHI/Private/D3D12CommandContext.h
+++ b/Engine/Source/Runtime/D3D12RHI/Private/D3D12CommandContext.h
@@ -615,7 +615,6 @@ public:
 
 	void SetSRVParameter(EShaderFrequency Frequency, uint32 SRVIndex, FD3D12ShaderResourceView* SRV);
 	void SetUAVParameter(EShaderFrequency Frequency, uint32 UAVIndex, FD3D12UnorderedAccessView* UAV);
-	void SetUAVParameter(EShaderFrequency Frequency, uint32 UAVIndex, FD3D12UnorderedAccessView* UAV, uint32 InitialCount);
 
 	void CommitGraphicsResourceTables();
 	void CommitComputeResourceTables();
diff --git a/Engine/Source/Runtime/D3D12RHI/Private/D3D12Commands.cpp b/Engine/Source/Runtime/D3D12RHI/Private/D3D12Commands.cpp
--- a/Engine/Source/Runtime/D3D12RHI/Private/D3D12Commands.cpp
+++ b/Engine/Source/Runtime/D3D12RHI/Private/D3D12Commands.cpp
@@ -451,12 +451,6 @@ void FD3D12CommandContext::SetUAVParameter(EShaderFrequency Frequency, uint32 UA
 	StateCache.SetUAV(Frequency, UAVIndex, UAV);
 }
 
-void FD3D12CommandContext::SetUAVParameter(EShaderFrequency Frequency, uint32 UAVIndex, FD3D12UnorderedAccessView* UAV, uint32 InitialCount)
-{
-	ClearShaderResources(UAV, EShaderParameterTypeMask::SRVMask);
-	StateCache.SetUAV(Frequency, UAVIndex, UAV, InitialCount);
-}
-
 void FD3D12CommandContext::SetSRVParameter(EShaderFrequency Frequency, uint32 SRVIndex, FD3D12ShaderResourceView* SRV)
 {
 	StateCache.SetShaderResourceView(Frequency, SRV, SRVIndex);
@@ -663,13 +657,13 @@ static void SetShaderParametersOnContext(
 	{
 		if (Parameter.Type == FRHIShaderParameterResource::EType::UnorderedAccessView)
 		{
-			if (ShaderFrequency == SF_Pixel || ShaderFrequency == SF_Vertex|| ShaderFrequency == SF_Compute)
+			if (D3D12DoesShaderFrequencySupportUAV(ShaderFrequency))
 			{
 				Binder.SetUAV(static_cast<FRHIUnorderedAccessView*>(Parameter.Resource), Parameter.Index, true);
 			}
 			else
 			{
-				checkf(false, TEXT("TShaderRHI Can't have compute shader to be set. UAVs are not supported tessellation and geometry shaders."));
+				checkf(false, TEXT("UAV set on unsupported shader frequency: %s"), GetShaderFrequencyString(ShaderFrequency));
 			}
 		}
 	}
@@ -754,13 +748,13 @@ static void SetShaderUnbindsOnContext(
 			Context.StateCache.SetShaderResourceView(ShaderFrequency, nullptr, Unbind.Index);
 			break;
 		case FRHIShaderParameterUnbind::EType::UnorderedAccessView:
-			if (ShaderFrequency == SF_Pixel || ShaderFrequency == SF_Vertex || ShaderFrequency == SF_Compute)
+			if (D3D12DoesShaderFrequencySupportUAV(ShaderFrequency))
 			{
 				Context.StateCache.SetUAV(ShaderFrequency, Unbind.Index, nullptr);
 			}
 			else
 			{
-				checkf(false, TEXT("TShaderRHI Can't have compute shader to be set. UAVs are not supported tessellation and geometry shaders."));
+				checkf(false, TEXT("UAV set on unsupported shader frequency: %s"), GetShaderFrequencyString(ShaderFrequency));
 			}
 			break;
 		default:
diff --git a/Engine/Source/Runtime/D3D12RHI/Private/D3D12DescriptorCache.cpp b/Engine/Source/Runtime/D3D12RHI/Private/D3D12DescriptorCache.cpp
--- a/Engine/Source/Runtime/D3D12RHI/Private/D3D12DescriptorCache.cpp
+++ b/Engine/Source/Runtime/D3D12RHI/Private/D3D12DescriptorCache.cpp
@@ -268,18 +268,18 @@ D3D12_GPU_DESCRIPTOR_HANDLE FD3D12DescriptorCache::BuildUAVTable(EShaderFrequenc
 
 void FD3D12DescriptorCache::SetUAVTable(EShaderFrequency ShaderStage, const FD3D12RootSignature* RootSignature, FD3D12UnorderedAccessViewCache& Cache, uint32 SlotsNeeded, const D3D12_GPU_DESCRIPTOR_HANDLE& BindDescriptor)
 {
-	check(ShaderStage == SF_Compute || ShaderStage == SF_Pixel || ShaderStage == SF_Vertex);
+	check(D3D12DoesShaderFrequencySupportUAV(ShaderStage));
 	const uint32 RootParameterIndex = RootSignature->UAVRDTBindSlot(ShaderStage);
 
 	check(FD3D12RootSignature::IsValidBindSlot(RootParameterIndex));
 
-	if (ShaderStage == SF_Pixel || ShaderStage == SF_Vertex)
+	if (ShaderStage == SF_Compute)
 	{
-		Context.GraphicsCommandList()->SetGraphicsRootDescriptorTable(RootParameterIndex, BindDescriptor);
+		Context.GraphicsCommandList()->SetComputeRootDescriptorTable(RootParameterIndex, BindDescriptor);
 	}
 	else
 	{
-		Context.GraphicsCommandList()->SetComputeRootDescriptorTable(RootParameterIndex, BindDescriptor);
+		Context.GraphicsCommandList()->SetGraphicsRootDescriptorTable(RootParameterIndex, BindDescriptor);
 	}
 
 	// We changed the descriptor table, so all resources bound to slots outside of the table's range are now dirty.
diff --git a/Engine/Source/Runtime/D3D12RHI/Private/D3D12RootSignature.h b/Engine/Source/Runtime/D3D12RHI/Private/D3D12RootSignature.h
--- a/Engine/Source/Runtime/D3D12RHI/Private/D3D12RootSignature.h
+++ b/Engine/Source/Runtime/D3D12RHI/Private/D3D12RootSignature.h
@@ -25,14 +25,18 @@ enum ERootParameterKeys
 	GS_CBVs,
 	GS_RootCBVs,
 	GS_Samplers,
+#if PLATFORM_SUPPORTS_MESH_SHADERS
 	MS_SRVs,
 	MS_CBVs,
 	MS_RootCBVs,
 	MS_Samplers,
+	MS_UAVs,
 	AS_SRVs,
 	AS_CBVs,
 	AS_RootCBVs,
 	AS_Samplers,
+	AS_UAVs,
+#endif
 	ALL_SRVs,
 	ALL_CBVs,
 	ALL_RootCBVs,
@@ -112,8 +116,10 @@ public:
 		switch (ShaderStage)
 		{
 		case SF_Vertex: return BindSlotMap[VS_Samplers];
+#if PLATFORM_SUPPORTS_MESH_SHADERS
 		case SF_Mesh: return BindSlotMap[MS_Samplers];
 		case SF_Amplification: return BindSlotMap[AS_Samplers];
+#endif
 		case SF_Pixel: return BindSlotMap[PS_Samplers];
 		case SF_Geometry: return BindSlotMap[GS_Samplers];
 		case SF_Compute: return BindSlotMap[ALL_Samplers];
@@ -128,8 +134,10 @@ public:
 		switch (ShaderStage)
 		{
 		case SF_Vertex: return BindSlotMap[VS_SRVs];
+#if PLATFORM_SUPPORTS_MESH_SHADERS
 		case SF_Mesh: return BindSlotMap[MS_SRVs];
 		case SF_Amplification: return BindSlotMap[AS_SRVs];
+#endif
 		case SF_Pixel: return BindSlotMap[PS_SRVs];
 		case SF_Geometry: return BindSlotMap[GS_SRVs];
 		case SF_Compute: return BindSlotMap[ALL_SRVs];
@@ -144,8 +152,10 @@ public:
 		switch (ShaderStage)
 		{
 		case SF_Vertex: return BindSlotMap[VS_CBVs];
+#if PLATFORM_SUPPORTS_MESH_SHADERS
 		case SF_Mesh: return BindSlotMap[MS_CBVs];
 		case SF_Amplification: return BindSlotMap[AS_CBVs];
+#endif
 		case SF_Pixel: return BindSlotMap[PS_CBVs];
 		case SF_Geometry: return BindSlotMap[GS_CBVs];
 		case SF_Compute: return BindSlotMap[ALL_CBVs];
@@ -160,8 +170,10 @@ public:
 		switch (ShaderStage)
 		{
 		case SF_Vertex: return BindSlotMap[VS_RootCBVs];
+#if PLATFORM_SUPPORTS_MESH_SHADERS
 		case SF_Mesh: return BindSlotMap[MS_RootCBVs];
 		case SF_Amplification: return BindSlotMap[AS_RootCBVs];
+#endif
 		case SF_Pixel: return BindSlotMap[PS_RootCBVs];
 		case SF_Geometry: return BindSlotMap[GS_RootCBVs];
 
@@ -181,9 +193,21 @@ public:
 
 	inline uint32 UAVRDTBindSlot(EShaderFrequency ShaderStage) const
 	{
-		check(ShaderStage == SF_Pixel || ShaderStage == SF_Vertex || ShaderStage == SF_Compute);
-		const uint32 MapSlotIndex = ShaderStage == SF_Vertex ? VS_UAVs : ALL_UAVs;
-		return BindSlotMap[MapSlotIndex];
+		switch (ShaderStage)
+		{
+		case SF_Vertex: return BindSlotMap[VS_UAVs];
+#if PLATFORM_SUPPORTS_MESH_SHADERS
+		case SF_Mesh: return BindSlotMap[MS_UAVs];
+		case SF_Amplification: return BindSlotMap[AS_UAVs];
+#endif
+		case SF_Pixel: return BindSlotMap[ALL_UAVs];
+
+		case SF_NumFrequencies:
+		case SF_Compute: return BindSlotMap[ALL_UAVs];
+
+		default: check(false);
+			return UINT_MAX;
+		}
 	}
 
 	inline static bool IsValidBindSlot(uint32 BindSlotMapIndex)
@@ -246,8 +270,10 @@ private:
 		switch (SF)
 		{
 		case SF_Vertex: pBindSlot = &BindSlotMap[VS_Samplers]; break;
+#if PLATFORM_SUPPORTS_MESH_SHADERS
 		case SF_Mesh: pBindSlot = &BindSlotMap[MS_Samplers]; break;
 		case SF_Amplification: pBindSlot = &BindSlotMap[AS_Samplers]; break;
+#endif
 		case SF_Pixel: pBindSlot = &BindSlotMap[PS_Samplers]; break;
 		case SF_Geometry: pBindSlot = &BindSlotMap[GS_Samplers]; break;
 
@@ -270,8 +296,10 @@ private:
 		switch (SF)
 		{
 		case SF_Vertex: pBindSlot = &BindSlotMap[VS_SRVs]; break;
+#if PLATFORM_SUPPORTS_MESH_SHADERS
 		case SF_Mesh: pBindSlot = &BindSlotMap[MS_SRVs]; break;
 		case SF_Amplification: pBindSlot = &BindSlotMap[AS_SRVs]; break;
+#endif
 		case SF_Pixel: pBindSlot = &BindSlotMap[PS_SRVs]; break;
 		case SF_Geometry: pBindSlot = &BindSlotMap[GS_SRVs]; break;
 
@@ -294,8 +322,10 @@ private:
 		switch (SF)
 		{
 		case SF_Vertex: pBindSlot = &BindSlotMap[VS_CBVs]; break;
+#if PLATFORM_SUPPORTS_MESH_SHADERS
 		case SF_Mesh: pBindSlot = &BindSlotMap[MS_CBVs]; break;
 		case SF_Amplification: pBindSlot = &BindSlotMap[AS_CBVs]; break;
+#endif
 		case SF_Pixel: pBindSlot = &BindSlotMap[PS_CBVs]; break;
 		case SF_Geometry: pBindSlot = &BindSlotMap[GS_CBVs]; break;
 
@@ -318,8 +348,10 @@ private:
 		switch (SF)
 		{
 		case SF_Vertex: pBindSlot = &BindSlotMap[VS_RootCBVs]; break;
+#if PLATFORM_SUPPORTS_MESH_SHADERS
 		case SF_Mesh: pBindSlot = &BindSlotMap[MS_RootCBVs]; break;
 		case SF_Amplification: pBindSlot = &BindSlotMap[AS_RootCBVs]; break;
+#endif
 		case SF_Pixel: pBindSlot = &BindSlotMap[PS_RootCBVs]; break;
 		case SF_Geometry: pBindSlot = &BindSlotMap[GS_RootCBVs]; break;
 
@@ -338,10 +370,22 @@ private:
 
 	inline void SetUAVRDTBindSlot(EShaderFrequency SF, uint8 RootParameterIndex)
 	{
-		check(SF == SF_Pixel || SF == SF_Vertex || SF == SF_Compute || SF == SF_NumFrequencies);
+		uint8* pBindSlot = nullptr;
+		switch (SF)
+		{
+		case SF_Vertex: pBindSlot = &BindSlotMap[VS_UAVs]; break;
+#if PLATFORM_SUPPORTS_MESH_SHADERS
+		case SF_Mesh: pBindSlot = &BindSlotMap[MS_UAVs]; break;
+		case SF_Amplification: pBindSlot = &BindSlotMap[AS_UAVs]; break;
+#endif
+		case SF_Pixel: pBindSlot = &BindSlotMap[ALL_UAVs]; break;
+
+		case SF_Compute:
+		case SF_NumFrequencies: pBindSlot = &BindSlotMap[ALL_UAVs]; break;
 
-		const uint32 MapSlotIndex = SF == SF_Vertex ? VS_UAVs : ALL_UAVs;
-		uint8* pBindSlot = &BindSlotMap[MapSlotIndex];
+		default: check(false);
+			return;
+		}
 
 		check(*pBindSlot == InvalidBindSlotMapIndex);
 		*pBindSlot = RootParameterIndex;
diff --git a/Engine/Source/Runtime/D3D12RHI/Private/D3D12StateCache.cpp b/Engine/Source/Runtime/D3D12RHI/Private/D3D12StateCache.cpp
--- a/Engine/Source/Runtime/D3D12RHI/Private/D3D12StateCache.cpp
+++ b/Engine/Source/Runtime/D3D12RHI/Private/D3D12StateCache.cpp
@@ -664,50 +664,44 @@ void FD3D12StateCache::ApplyResourceTables(const FD3D12RootSignature* const pRoo
 	CBVSlotMask CurrentShaderDirtyCBVSlots[SF_NumStandardFrequencies] = {};
 #endif
 	UAVSlotMask CurrentShaderDirtyUAVSlots[SF_NumStandardFrequencies] = {};
-	uint32 NumUAVs = 0;
+	uint32 NumUAVs[SF_NumStandardFrequencies] = {};
 	uint32 NumSRVs[SF_NumStandardFrequencies] = {};
 #if D3D12RHI_USE_CONSTANT_BUFFER_VIEWS
-	uint32 NumCBVs[SF_NumStandardFrequencies] ={};
+	uint32 NumCBVs[SF_NumStandardFrequencies] = {};
 #endif
 	uint32 NumViews = 0;
 
-	EShaderFrequency ComputeUAVStages[] = {SF_Compute};
-	EShaderFrequency GraphicsUAVStages[] = {SF_Vertex, SF_Pixel};
-
-	TArrayView<EShaderFrequency> UAVStages = StartStage == SF_Compute ? MakeArrayView(ComputeUAVStages) : MakeArrayView(GraphicsUAVStages);
-
 	for (uint32 iTries = 0; iTries < 2; ++iTries)
 	{
-		if (bUAVs)
+		for (uint32 Stage = StartStage; Stage < EndStage; ++Stage)
 		{
-			for (EShaderFrequency UAVStage : UAVStages)
+			if (ShouldSkipStage(Stage))
 			{
-				const UAVSlotMask CurrentShaderUAVRegisterMask = BitMask<UAVSlotMask>(PipelineState.Common.CurrentShaderUAVCounts[UAVStage]);
-				CurrentShaderDirtyUAVSlots[UAVStage] = CurrentShaderUAVRegisterMask & PipelineState.Common.UAVCache.DirtySlotMask[UAVStage];
-				if (CurrentShaderDirtyUAVSlots[UAVStage])
+				continue;
+			}
+
+			if (bUAVs)
+			{
+				const UAVSlotMask CurrentShaderUAVRegisterMask = BitMask<UAVSlotMask>(PipelineState.Common.CurrentShaderUAVCounts[Stage]);
+				CurrentShaderDirtyUAVSlots[Stage] = CurrentShaderUAVRegisterMask & PipelineState.Common.UAVCache.DirtySlotMask[Stage];
+				if (CurrentShaderDirtyUAVSlots[Stage])
 				{
+					check(D3D12DoesShaderFrequencySupportUAV(static_cast<EShaderFrequency>(Stage)));
+
 					if (ResourceBindingTier <= D3D12_RESOURCE_BINDING_TIER_2)
 					{
 						// Tier 1 and 2 HW requires the full number of UAV descriptors defined in the root signature's descriptor table.
-						NumUAVs = pRootSignature->MaxUAVCount(UAVStage);
+						NumUAVs[Stage] = pRootSignature->MaxUAVCount(Stage);
 					}
 					else
 					{
-						NumUAVs = PipelineState.Common.CurrentShaderUAVCounts[UAVStage];
+						NumUAVs[Stage] = PipelineState.Common.CurrentShaderUAVCounts[Stage];
 					}
 
-					check(NumUAVs > 0 && NumUAVs <= MAX_UAVS);
-					NumViews += NumUAVs;
+					check(NumUAVs[Stage] > 0 && NumUAVs[Stage] <= MAX_UAVS);
+					NumViews += NumUAVs[Stage];
 				}
 			}
-		}
-
-		for (uint32 Stage = StartStage; Stage < EndStage; ++Stage)
-		{
-			if (ShouldSkipStage(Stage))
-			{
-				continue;
-			}
 
 			if (bSRVs)
 			{
@@ -775,14 +769,19 @@ void FD3D12StateCache::ApplyResourceTables(const FD3D12RootSignature* const pRoo
 	uint32 ViewHeapSlot = DescriptorCache.GetCurrentViewHeap()->ReserveSlots(NumViews);
 
 	// Unordered access views
+	if (bUAVs)
 	{
-		SCOPE_CYCLE_COUNTER(STAT_D3D12ApplyStateSetUAVTime);
-		for (EShaderFrequency UAVStage : UAVStages)
+		//SCOPE_CYCLE_COUNTER(STAT_D3D12ApplyStateSetUAVTime);
+		FD3D12UnorderedAccessViewCache& UAVCache = PipelineState.Common.UAVCache;
+
+		for (uint32 Index = StartStage; Index < EndStage; Index++)
 		{
-			if (CurrentShaderDirtyUAVSlots[UAVStage])
+			if (CurrentShaderDirtyUAVSlots[Index])
 			{
-				const D3D12_GPU_DESCRIPTOR_HANDLE BindDescriptor = DescriptorCache.BuildUAVTable(UAVStage, pRootSignature, PipelineState.Common.UAVCache, CurrentShaderDirtyUAVSlots[UAVStage], NumUAVs, ViewHeapSlot);
-				DescriptorCache.SetUAVTable(UAVStage, pRootSignature, PipelineState.Common.UAVCache, NumUAVs, BindDescriptor);
+				const EShaderFrequency ShaderFrequency = static_cast<EShaderFrequency>(Index);
+
+				const D3D12_GPU_DESCRIPTOR_HANDLE BindDescriptor = DescriptorCache.BuildUAVTable(ShaderFrequency, pRootSignature, UAVCache, CurrentShaderDirtyUAVSlots[Index], NumUAVs[Index], ViewHeapSlot);
+				DescriptorCache.SetUAVTable(ShaderFrequency, pRootSignature, UAVCache, NumUAVs[Index], BindDescriptor);
 			}
 		}
 	}
@@ -1306,7 +1305,6 @@ void FD3D12StateCache::SetRootConstants(const FUint32Vector4& Constants)
 void FD3D12StateCache::ClearUAVs(EShaderFrequency ShaderStage)
 {
 	FD3D12UnorderedAccessViewCache& Cache = PipelineState.Common.UAVCache;
-	const bool bIsCompute = ShaderStage == SF_Compute;
 
 	for (uint32 i = 0; i < MAX_UAVS; ++i)
 	{
@@ -1318,7 +1316,7 @@ void FD3D12StateCache::ClearUAVs(EShaderFrequency ShaderStage)
 	}
 }
 
-void FD3D12StateCache::SetUAV(EShaderFrequency ShaderStage, uint32 SlotIndex, FD3D12UnorderedAccessView* UAV, uint32 InitialCount)
+void FD3D12StateCache::SetUAV(EShaderFrequency ShaderStage, uint32 SlotIndex, FD3D12UnorderedAccessView* UAV)
 {
 	SCOPE_CYCLE_COUNTER(STAT_D3D12SetUnorderedAccessViewTime);
 
diff --git a/Engine/Source/Runtime/D3D12RHI/Private/D3D12StateCachePrivate.h b/Engine/Source/Runtime/D3D12RHI/Private/D3D12StateCachePrivate.h
--- a/Engine/Source/Runtime/D3D12RHI/Private/D3D12StateCachePrivate.h
+++ b/Engine/Source/Runtime/D3D12RHI/Private/D3D12StateCachePrivate.h
@@ -705,7 +705,7 @@ public:
 
 	void SetRootConstants(const FUint32Vector4& Constants);
 
-	void SetUAV(EShaderFrequency ShaderStage, uint32 SlotIndex, FD3D12UnorderedAccessView* UAV, uint32 InitialCount = -1);
+	void SetUAV(EShaderFrequency ShaderStage, uint32 SlotIndex, FD3D12UnorderedAccessView* UAV);
 
 	void ClearUAVs(EShaderFrequency ShaderStage);
 
diff --git a/Engine/Source/Runtime/D3D12RHI/Private/D3D12Util.cpp b/Engine/Source/Runtime/D3D12RHI/Private/D3D12Util.cpp
--- a/Engine/Source/Runtime/D3D12RHI/Private/D3D12Util.cpp
+++ b/Engine/Source/Runtime/D3D12RHI/Private/D3D12Util.cpp
@@ -1154,7 +1154,6 @@ void FD3D12QuantizedBoundShaderState::InitShaderRegisterCounts(const D3D12_RESOU
 	// https://en.wikipedia.org/wiki/Feature_levels_in_Direct3D
 
 	// To reduce the size of the root signature, we only allow UAVs for certain shaders. 
-	// This code makes the assumption that the engine only uses UAVs at the PS or CS shader stages.
 	check(bAllowUAVs || (!bAllowUAVs && Counts.NumUAVs == 0));
 
 	if (ResourceBindingTier <= D3D12_RESOURCE_BINDING_TIER_1)
@@ -1393,8 +1392,8 @@ const FD3D12RootSignature* FD3D12Adapter::GetRootSignature(const FBoundShaderSta
 
 	QuantizeBoundShaderStateCommon(QBSS, FD3D12DynamicRHI::ResourceCast(BSS.GetVertexShader()),        ResourceBindingTier, SV_Vertex, true /*bAllowUAVs*/);
 #if PLATFORM_SUPPORTS_MESH_SHADERS
-	QuantizeBoundShaderStateCommon(QBSS, FD3D12DynamicRHI::ResourceCast(BSS.GetMeshShader()),          ResourceBindingTier, SV_Mesh);
-	QuantizeBoundShaderStateCommon(QBSS, FD3D12DynamicRHI::ResourceCast(BSS.GetAmplificationShader()), ResourceBindingTier, SV_Amplification);
+	QuantizeBoundShaderStateCommon(QBSS, FD3D12DynamicRHI::ResourceCast(BSS.GetMeshShader()),          ResourceBindingTier, SV_Mesh, true /*bAllowUAVs*/);
+	QuantizeBoundShaderStateCommon(QBSS, FD3D12DynamicRHI::ResourceCast(BSS.GetAmplificationShader()), ResourceBindingTier, SV_Amplification, true /*bAllowUAVs*/);
 #endif
 	QuantizeBoundShaderStateCommon(QBSS, FD3D12DynamicRHI::ResourceCast(BSS.GetPixelShader()),         ResourceBindingTier, SV_Pixel, true /*bAllowUAVs*/);
 	QuantizeBoundShaderStateCommon(QBSS, FD3D12DynamicRHI::ResourceCast(BSS.GetGeometryShader()),      ResourceBindingTier, SV_Geometry);
diff --git a/Engine/Source/Runtime/D3D12RHI/Private/D3D12Util.h b/Engine/Source/Runtime/D3D12RHI/Private/D3D12Util.h
--- a/Engine/Source/Runtime/D3D12RHI/Private/D3D12Util.h
+++ b/Engine/Source/Runtime/D3D12RHI/Private/D3D12Util.h
@@ -724,3 +724,14 @@ FORCEINLINE_DEBUGGABLE bool CheckResourceStateCompatibility(D3D12_RESOURCE_STATE
 	return true;
 }
 
+inline bool D3D12DoesShaderFrequencySupportUAV(EShaderFrequency ShaderFrequency)
+{
+	return ShaderFrequency == SF_Pixel
+		|| ShaderFrequency == SF_Vertex
+		|| ShaderFrequency == SF_Compute
+#if PLATFORM_SUPPORTS_MESH_SHADERS
+		|| ShaderFrequency == SF_Mesh
+		|| ShaderFrequency == SF_Amplification
+#endif
+		;
+}
