Fix crash on exit for -d3ddebug runs

#Synced-CL 50292737

diff --git a/Engine/Source/Runtime/D3D12RHI/Private/D3D12Adapter.cpp b/Engine/Source/Runtime/D3D12RHI/Private/D3D12Adapter.cpp
--- a/Engine/Source/Runtime/D3D12RHI/Private/D3D12Adapter.cpp
+++ b/Engine/Source/Runtime/D3D12RHI/Private/D3D12Adapter.cpp
@@ -268,16 +268,15 @@ static bool LogD3D12Message(D3D12_MESSAGE_SEVERITY Severity, D3D12_MESSAGE_ID ID
 	return bBreakOnLog;
 }
 
-static bool CheckD3DStoredMessages()
+static bool CheckD3DStoredMessages(ID3D12Device* Device)
 {
 	bool bResult = false;
 
 	TRefCountPtr<ID3D12Debug> d3dDebug;
 	if (SUCCEEDED(D3D12GetDebugInterface(__uuidof(ID3D12Debug), (void**)d3dDebug.GetInitReference())))
 	{
-		FD3D12DynamicRHI* D3D12RHI = FD3D12DynamicRHI::GetD3DRHI();
 		TRefCountPtr<ID3D12InfoQueue> d3dInfoQueue;
-		if (SUCCEEDED(D3D12RHI->GetAdapter().GetD3DDevice()->QueryInterface(__uuidof(ID3D12InfoQueue), (void**)d3dInfoQueue.GetInitReference())))
+		if (SUCCEEDED(Device->QueryInterface(__uuidof(ID3D12InfoQueue), (void**)d3dInfoQueue.GetInitReference())))
 		{
 			D3D12_MESSAGE* d3dMessage = nullptr;
 			SIZE_T AllocateSize = 0;
@@ -341,7 +340,9 @@ static LONG __stdcall D3DVectoredExceptionHandler(EXCEPTION_POINTERS* InInfo)
 		// Logging and break in code is managed in the D3D12MessageCallBack if ID3D12InfoQueue1 but we handle and continue D3D exception in order to not break while the debugger is not attached
 		if (GD3D12CallBackCookie == 0)
 		{
-			if (CheckD3DStoredMessages())
+			FD3D12DynamicRHI* const D3D12RHI = FD3D12DynamicRHI::GetD3DRHI();
+			if (D3D12RHI && D3D12RHI->HasAdapter() &&
+				CheckD3DStoredMessages(D3D12RHI->GetAdapter().GetD3DDevice()))
 			{
 				if (FPlatformMisc::IsDebuggerPresent())
 				{
@@ -1677,7 +1678,7 @@ FD3D12Adapter::~FD3D12Adapter()
 		DXGIDebug.SafeRelease();
 
 #if D3D12_SUPPORTS_INFO_QUEUE
-		CheckD3DStoredMessages();
+		CheckD3DStoredMessages(RootDevice);
 
 		if (GD3D12CallBackCookie != 0)
 		{
