When moving a D3D12ResourceLocation with a pool allocation, make sure the Memmove happens inside the FD3D12PoolAllocator critical section to maintain thread safety with its internal linked lists.

#Synced-CL 50267093

diff --git a/Engine/Source/Runtime/D3D12RHI/Private/D3D12PoolAllocator.cpp b/Engine/Source/Runtime/D3D12RHI/Private/D3D12PoolAllocator.cpp
--- a/Engine/Source/Runtime/D3D12RHI/Private/D3D12PoolAllocator.cpp
+++ b/Engine/Source/Runtime/D3D12RHI/Private/D3D12PoolAllocator.cpp
@@ -851,6 +851,9 @@ void FD3D12PoolAllocator::TransferOwnership(FD3D12ResourceLocation& InSource, FD
 
 	check(IsOwner(InSource));
 
+	// Note: mirrors FD3D12ResourceLocation::TransferOwnership but is inside our CS.
+	FMemory::Memmove(&InDest, &InSource, sizeof(FD3D12ResourceLocation));
+
 	FRHIPoolAllocationData& SourcePoolData = InSource.GetPoolAllocatorPrivateData().PoolData;
 	FRHIPoolAllocationData& DestinationPoolData = InDest.GetPoolAllocatorPrivateData().PoolData;
 
diff --git a/Engine/Source/Runtime/D3D12RHI/Private/D3D12Resources.cpp b/Engine/Source/Runtime/D3D12RHI/Private/D3D12Resources.cpp
--- a/Engine/Source/Runtime/D3D12RHI/Private/D3D12Resources.cpp
+++ b/Engine/Source/Runtime/D3D12RHI/Private/D3D12Resources.cpp
@@ -1217,12 +1217,15 @@ void FD3D12ResourceLocation::TransferOwnership(FD3D12ResourceLocation& Destinati
 	// Clear out the destination
 	Destination.Clear();
 
-	FMemory::Memmove(&Destination, &Source, sizeof(FD3D12ResourceLocation));
-
 	if (Source.GetAllocatorType() == FD3D12ResourceLocation::AT_Pool)
 	{
+		// Note: pool allocated locations handle the copy themselves to prevent data races when updating their linked list.
 		Source.GetPoolAllocator()->TransferOwnership(Source, Destination);
 	}
+	else
+	{
+		FMemory::Memmove(&Destination, &Source, sizeof(FD3D12ResourceLocation));
+	}
 
 	// update tracked allocation
 #if !PLATFORM_WINDOWS && ENABLE_LOW_LEVEL_MEM_TRACKER
