Minor cleanup of FD3D12DynamicRHI::RHIAsyncCreateTexture2D
- Use move semantics when transferring the ResourceLocations to tasks.
- Use UE::RHITextureUtils instead of manual texture size calculations.

#Synced-CL 50544822

diff --git a/Engine/Source/Runtime/D3D12RHI/Private/D3D12Texture.cpp b/Engine/Source/Runtime/D3D12RHI/Private/D3D12Texture.cpp
--- a/Engine/Source/Runtime/D3D12RHI/Private/D3D12Texture.cpp
+++ b/Engine/Source/Runtime/D3D12RHI/Private/D3D12Texture.cpp
@@ -1066,20 +1066,6 @@ static UE::RHICore::FBaseTextureInitializerImplementation CreateTextureInitializ
 	);
 }
 
-uint64 CalculateTextureSize(const FRHITextureDesc& Desc)
-{
-	const FPixelFormatInfo& PixelFormat = GPixelFormats[Desc.Format];
-
-	const uint64 BlockSizeX = FMath::DivideAndRoundUp<uint64>(Desc.Extent.X, PixelFormat.BlockSizeX);
-	const uint64 BlockSizeY = FMath::DivideAndRoundUp<uint64>(Desc.Extent.Y, PixelFormat.BlockSizeY);
-	const uint64 BlockSizeZ = FMath::DivideAndRoundUp<uint64>(Desc.Depth,    PixelFormat.BlockSizeZ);
-	const uint64 BlockCount = BlockSizeX * BlockSizeY * BlockSizeZ;
-	const uint64 ResourceSize = BlockCount * PixelFormat.BlockBytes;
-	const uint64 ResourceCount = Desc.GetSubresourceCount();
-
-	return ResourceSize * ResourceCount;
-}
-
 FRHITextureInitializer FD3D12DynamicRHI::RHICreateTextureInitializer(FRHICommandListBase& RHICmdList, const FRHITextureCreateDesc& CreateDesc)
 {
 	TRACE_CPUPROFILER_EVENT_SCOPE(D3D12RHI::RHICreateTextureInitializer);
@@ -1132,14 +1118,12 @@ public:
 
 	FWaitInitialMipDataUploadTask(
 		FD3D12Texture* InTexture,
-		FD3D12ResourceLocation& InTempResourceLocation,
-		FD3D12ResourceLocation& InTempResourceLocationLowMips)
+		FD3D12ResourceLocation&& InTempResourceLocation,
+		FD3D12ResourceLocation&& InTempResourceLocationLowMips)
 		: Texture(InTexture)
-		, TempResourceLocation(InTempResourceLocation.GetParentDevice())
-		, TempResourceLocationLowMips(InTempResourceLocationLowMips.GetParentDevice())
+		, TempResourceLocation(MoveTemp(InTempResourceLocation))
+		, TempResourceLocationLowMips(MoveTemp(InTempResourceLocationLowMips))
 	{
-		FD3D12ResourceLocation::TransferOwnership(TempResourceLocation, InTempResourceLocation);
-		FD3D12ResourceLocation::TransferOwnership(TempResourceLocationLowMips, InTempResourceLocationLowMips);
 	}
 
 	void DoTask(ENamedThreads::Type CurrentThread, const FGraphEventRef& MyCompletionGraphEvent)
@@ -1195,25 +1179,24 @@ FTextureRHIRef FD3D12DynamicRHI::RHIAsyncCreateTexture2D(uint32 SizeX, uint32 Si
 	D3D12_SUBRESOURCE_DATA SubResourceData[MAX_TEXTURE_MIP_COUNT] = { };
 	for (uint32 MipIndex = 0; MipIndex < NumInitialMips; ++MipIndex)
 	{
-		uint32 NumBlocksX = FMath::Max<uint32>(1, ((SizeX >> MipIndex) + GPixelFormats[Format].BlockSizeX-1) / GPixelFormats[Format].BlockSizeX);
-		uint32 NumBlocksY = FMath::Max<uint32>(1, ((SizeY >> MipIndex) + GPixelFormats[Format].BlockSizeY-1) / GPixelFormats[Format].BlockSizeY);
+		uint64 RowPitch = 0;
+		uint64 MipSize = UE::RHITextureUtils::CalculateTextureMipSize(CreateDesc, MipIndex, RowPitch);
 
 		SubResourceData[MipIndex].pData = InitialMipData[MipIndex];
-		SubResourceData[MipIndex].RowPitch = NumBlocksX * GPixelFormats[Format].BlockBytes;
-		SubResourceData[MipIndex].SlicePitch = NumBlocksX * NumBlocksY * GPixelFormats[Format].BlockBytes;
+		SubResourceData[MipIndex].RowPitch = RowPitch;
+		SubResourceData[MipIndex].SlicePitch = MipSize;
 	}
 
 	void* TempBuffer = ZeroBuffer;
-	uint32 TempBufferSize = ZeroBufferSize;
+	uint64 TempBufferSize = ZeroBufferSize;
 	for (uint32 MipIndex = NumInitialMips; MipIndex < NumMips; ++MipIndex)
 	{
-		uint32 NumBlocksX = FMath::Max<uint32>(1, ((SizeX >> MipIndex) + GPixelFormats[Format].BlockSizeX-1) / GPixelFormats[Format].BlockSizeX);
-		uint32 NumBlocksY = FMath::Max<uint32>(1, ((SizeY >> MipIndex) + GPixelFormats[Format].BlockSizeY-1) / GPixelFormats[Format].BlockSizeY);
-		uint32 MipSize = NumBlocksX * NumBlocksY * GPixelFormats[Format].BlockBytes;
+		uint64 RowPitch = 0;
+		uint64 MipSize = UE::RHITextureUtils::CalculateTextureMipSize(CreateDesc, MipIndex, RowPitch);
 
 		if (MipSize > TempBufferSize)
 		{
-			UE_LOG(LogD3D12RHI, Verbose, TEXT("Temp texture streaming buffer not large enough, needed %d bytes"), MipSize);
+			UE_LOG(LogD3D12RHI, Verbose, TEXT("Temp texture streaming buffer not large enough, needed %llu bytes"), MipSize);
 			check(TempBufferSize == ZeroBufferSize);
 			TempBufferSize = MipSize;
 			TempBuffer = FMemory::Malloc(TempBufferSize);
@@ -1221,7 +1204,7 @@ FTextureRHIRef FD3D12DynamicRHI::RHIAsyncCreateTexture2D(uint32 SizeX, uint32 Si
 		}
 
 		SubResourceData[MipIndex].pData = TempBuffer;
-		SubResourceData[MipIndex].RowPitch = NumBlocksX * GPixelFormats[Format].BlockBytes;
+		SubResourceData[MipIndex].RowPitch = RowPitch;
 		SubResourceData[MipIndex].SlicePitch = MipSize;
 	}
 
@@ -1419,8 +1402,8 @@ FTextureRHIRef FD3D12DynamicRHI::RHIAsyncCreateTexture2D(uint32 SizeX, uint32 Si
 
 		OutCompletionEvent = TGraphTask<FWaitInitialMipDataUploadTask>::CreateTask(&CopyCompleteEvents).ConstructAndDispatchWhenReady(
 			TextureOut,
-			TempResourceLocation,
-			TempResourceLocationLowMips);
+			MoveTemp(TempResourceLocation),
+			MoveTemp(TempResourceLocationLowMips));
 	}
 
 	if (TempBufferSize != ZeroBufferSize)
