Disable shader execution reordering support through NVAPI when Pix is attached for capture

#Synced-CL 50609014

diff --git a/Engine/Source/Runtime/D3D12RHI/Private/Windows/WindowsD3D12Device.cpp b/Engine/Source/Runtime/D3D12RHI/Private/Windows/WindowsD3D12Device.cpp
--- a/Engine/Source/Runtime/D3D12RHI/Private/Windows/WindowsD3D12Device.cpp
+++ b/Engine/Source/Runtime/D3D12RHI/Private/Windows/WindowsD3D12Device.cpp
@@ -1633,20 +1633,34 @@ void FD3D12DynamicRHI::Init()
 					}
 				}
 #endif // WITH_NVAPI_RAYTRACING_VALIDATION
+	
+				bool bSkipSERSupportCheck = false;
 
-				NvAPI_Status NvStatusSER = NvAPI_D3D12_IsNvShaderExtnOpCodeSupported(GetAdapter().GetD3DDevice(), NV_EXTN_OP_HIT_OBJECT_REORDER_THREAD, &GRHIGlobals.SupportsShaderExecutionReordering);
-				if (NvStatusSER == NVAPI_OK && GRHIGlobals.SupportsShaderExecutionReordering)
+#if D3D12RHI_SUPPORTS_WIN_PIX
+				// Pix only supports captures with SER enabled when using D3D12's official SER API, which is currently not available in a non-preview Agility SDK
+				if (PIXIsAttachedForGpuCapture())
 				{
-					NVAPI_D3D12_RAYTRACING_THREAD_REORDERING_CAPS ReorderCaps = NVAPI_D3D12_RAYTRACING_THREAD_REORDERING_CAP_NONE;
-					NvStatusSER = NvAPI_D3D12_GetRaytracingCaps(
-						GetAdapter().GetD3DDevice(),
-						NVAPI_D3D12_RAYTRACING_CAPS_TYPE_THREAD_REORDERING,
-						&ReorderCaps,
-						sizeof(ReorderCaps));
-
-					if (NvStatusSER != NVAPI_OK || ReorderCaps == NVAPI_D3D12_RAYTRACING_THREAD_REORDERING_CAP_NONE)
+					bSkipSERSupportCheck = true;
+					GRHIGlobals.SupportsShaderExecutionReordering = false;
+					UE_LOG(LogD3D12RHI, Log, TEXT("Skipping NvAPI Shader Execution Reordering support as Pix is attached to the process."));
+				}
+#endif
+				if (!bSkipSERSupportCheck)
+				{
+					NvAPI_Status NvStatusSER = NvAPI_D3D12_IsNvShaderExtnOpCodeSupported(GetAdapter().GetD3DDevice(), NV_EXTN_OP_HIT_OBJECT_REORDER_THREAD, &GRHIGlobals.SupportsShaderExecutionReordering);
+					if (NvStatusSER == NVAPI_OK && GRHIGlobals.SupportsShaderExecutionReordering)
 					{
-						GRHIGlobals.SupportsShaderExecutionReordering = false;
+						NVAPI_D3D12_RAYTRACING_THREAD_REORDERING_CAPS ReorderCaps = NVAPI_D3D12_RAYTRACING_THREAD_REORDERING_CAP_NONE;
+						NvStatusSER = NvAPI_D3D12_GetRaytracingCaps(
+							GetAdapter().GetD3DDevice(),
+							NVAPI_D3D12_RAYTRACING_CAPS_TYPE_THREAD_REORDERING,
+							&ReorderCaps,
+							sizeof(ReorderCaps));
+
+						if (NvStatusSER != NVAPI_OK || ReorderCaps == NVAPI_D3D12_RAYTRACING_THREAD_REORDERING_CAP_NONE)
+						{
+							GRHIGlobals.SupportsShaderExecutionReordering = false;
+						}
 					}
 				}
 
