Adding an asynchronous pool allocation strategy for D3D12 buffer allocators (Upload heap and Default buffer pools).
The implementation pre-allocates "Overflow" pools on a background thread in order to reduce allocation hitches on the critical path that occur when existing pools become full and new pools must be created synchronously.

The allocator first attempts to use a pre-allocated overflow pool, if no overflow pool is available, it falls back to synchronous allocation.
After each allocation, the system checks if additional overflow pools should be pre-allocated on a background thread.

#Synced-CL 50589146

diff --git a/Engine/Source/Runtime/D3D12RHI/Private/D3D12Allocation.cpp b/Engine/Source/Runtime/D3D12RHI/Private/D3D12Allocation.cpp
--- a/Engine/Source/Runtime/D3D12RHI/Private/D3D12Allocation.cpp
+++ b/Engine/Source/Runtime/D3D12RHI/Private/D3D12Allocation.cpp
@@ -13,6 +13,7 @@
 #include "HAL/LowLevelMemTracker.h"
 #include "HAL/LowLevelMemStats.h"
 #include "ProfilingDebugging/MemoryTrace.h"
+#include "Async/Async.h"
 
 // Fix for random GPU crashes on draw indirects on multiple IHVs. Force all indirect arg buffers as committed resources (see UE-115982)
 static int32 GD3D12AllowPoolAllocateIndirectArgBuffers = 1;
@@ -124,13 +125,56 @@ static FAutoConsoleVariableRef CVarD3D12UploadHeapBigBlockMaxAllocationSize(
 	TEXT("Maximum allocation size on the big block allocator for upload memory"),
 	ECVF_ReadOnly);
 
-static int32 GD3D12UploadHeapBigBlockPoolSize = 8 * 1024 * 1024;
+static int32 GD3D12UploadHeapBigBlockPoolSize = 32 * 1024 * 1024; //On a game with Nanite disabled this value could be lowered
 static FAutoConsoleVariableRef CVarD3D12UploadHeapBigBlockPoolSize(
 	TEXT("d3d12.UploadHeap.BigBlock.PoolSize"),
 	GD3D12UploadHeapBigBlockPoolSize,
 	TEXT("Pool size for the upload memory big block allocator"),
 	ECVF_ReadOnly);
 
+#if D3D12RHI_PLATFORM_HAS_UNIFIED_MEMORY
+static int32 GD3D12UploadHeapPoolAllocationStrategy = 0;
+#else
+static int32 GD3D12UploadHeapPoolAllocationStrategy = 1;
+#endif
+
+static FAutoConsoleVariableRef CVarD3D12UploadHeapPoolAllocationStrategy(
+	TEXT("d3d12.UploadHeap.PoolAllocationStrategy"),
+	GD3D12UploadHeapPoolAllocationStrategy,
+	TEXT("Pool allocation strategy")
+	TEXT("0 - Immediate: a new pool will be allocated on the spot when needed")
+	TEXT("1 - Async: the system will try to pre allocate overflow pools in advance"),
+	ECVF_ReadOnly);
+
+static int32 GD3D12UploadHeapOverFlowPoolsCount = 1;
+static FAutoConsoleVariableRef CVarD3D12UploadHeapOverFlowPoolsCount(
+	TEXT("d3d12.UploadHeap.OverFlowPoolsCount"),
+	GD3D12UploadHeapOverFlowPoolsCount,
+	TEXT("Max OverFlow Pools count to pre-allocate"),
+	ECVF_ReadOnly);
+
+#if D3D12RHI_PLATFORM_HAS_UNIFIED_MEMORY
+static int32 GD3D12DefaultBufferPoolAllocationStrategy = 0;
+#else
+static int32 GD3D12DefaultBufferPoolAllocationStrategy = 1;
+#endif
+
+static FAutoConsoleVariableRef CVarD3D12DefaultBufferPoolAllocationStrategy(
+	TEXT("d3d12.DefaultBuffer.PoolAllocationStrategy"),
+	GD3D12DefaultBufferPoolAllocationStrategy,
+	TEXT("Pool allocation strategy")
+	TEXT("0 - Immediate: a new pool will be allocated on the spot when needed")
+	TEXT("1 - Async: the system will try to pre allocate overflow pools in advance"),
+	ECVF_ReadOnly);
+
+static int32 GD3D12DefaultBufferOverFlowPoolsCount = 1;
+static FAutoConsoleVariableRef CVarD3D12DefaultBufferPoolOverFlowPoolsCount(
+	TEXT("d3d12.DefaultBuffer.OverFlowPoolsCount"),
+	GD3D12DefaultBufferOverFlowPoolsCount,
+	TEXT("Max OverFlow Pools count to pre-allocate"),
+	ECVF_ReadOnly);
+
+
 static int32 GD3D12FastConstantAllocatorPageSize = 64 * 1024;
 static FAutoConsoleVariableRef CVarD3D12FastConstantAllocatorPageSize(
 	TEXT("d3d12.FastConstantAllocatorPageSize"),
@@ -161,6 +205,13 @@ namespace ED3D12AllocatorID
 	};
 };
 
+static ERHIPoolAllocationStrategy GetPoolAllocationStrategy(uint32 AllocationStrategyIndex)
+{
+	// 0 = ERHIPoolAllocationStrategy::Immediate
+	// 1 = ERHIPoolAllocationStrategy::Async
+	return (AllocationStrategyIndex == 0) ? ERHIPoolAllocationStrategy::Immediate : ERHIPoolAllocationStrategy::Async;
+}
+
 //-----------------------------------------------------------------------------
 //	Allocator Base
 //-----------------------------------------------------------------------------
@@ -722,11 +773,15 @@ FD3D12MultiBuddyAllocator::FD3D12MultiBuddyAllocator(FD3D12Device* ParentDevice,
 	uint32 InMaxAllocationSize,
 	uint32 InDefaultPoolSize,
 	uint32 InMinBlockSize,
-	HeapId InTraceParentHeapId)
+	HeapId InTraceParentHeapId,
+	ERHIPoolAllocationStrategy InPoolAllocationStrategy,
+	uint16 InOverFlowPoolsCount)
 	: FD3D12ResourceAllocator(ParentDevice, VisibleNodes, InInitConfig, Name, InMaxAllocationSize)
 	, AllocationStrategy(InAllocationStrategy)
 	, MinBlockSize(InMinBlockSize)
 	, DefaultPoolSize(InDefaultPoolSize)
+	, PoolAllocationStrategy(InPoolAllocationStrategy)
+	, OverFlowPoolsCount(InOverFlowPoolsCount)
 {
 #if UE_MEMORY_TRACE_ENABLED
 	TraceHeapId = MemoryTrace_HeapSpec(InTraceParentHeapId, *Name);
@@ -740,18 +795,82 @@ FD3D12MultiBuddyAllocator::~FD3D12MultiBuddyAllocator()
 
 bool FD3D12MultiBuddyAllocator::TryAllocate(uint32 SizeInBytes, uint32 Alignment, FD3D12ResourceLocation& ResourceLocation)
 {
-	FScopeLock Lock(&CS);
-
-	for (int32 i = 0; i < Allocators.Num(); i++)
+	if (PoolAllocationStrategy == ERHIPoolAllocationStrategy::Async)
 	{
-		if (Allocators[i]->TryAllocate(SizeInBytes, Alignment, ResourceLocation))
+		CS.Lock();
+
+		// No Pools, needs initialization
+		if (Allocators.IsEmpty())
 		{
-			return true;
+			Initialize();
 		}
+
+		// Iterate over Pool and find one where we can allocate from 
+		bool AllocationSuccess = false;
+
+		for (int32 i = 0; i < Allocators.Num(); i++)
+		{
+			if (Allocators[i]->TryAllocate(SizeInBytes, Alignment, ResourceLocation))
+			{
+				AllocationSuccess = true;
+				break;
+			}
+		}
+		
+		// Handle a failed allocation
+		if (!AllocationSuccess)
+		{
+			// Allocation failed, grab a ready OverFlow pool and try allocating from it
+			if (OverFlowAllocators.Num())
+			{
+				Allocators.Add(OverFlowAllocators.Pop());
+				AllocationSuccess = Allocators.Last()->TryAllocate(SizeInBytes, Alignment, ResourceLocation);
+			}
+
+			// Allocation failed and either no OverFlow pool is ready or allocation failed from it as well, we need to allocate a new Pool immediately
+			// we want to avoid this case as much as possible, consider increasing OverFlowPoolsCount if this happens
+			if (!AllocationSuccess)
+			{
+				Allocators.Add(CreateNewAllocator(SizeInBytes));
+				AllocationSuccess = Allocators.Last()->TryAllocate(SizeInBytes, Alignment, ResourceLocation);
+			}
+		}
+			
+		// Check if we needs to start a new async Pool Allocation if not in progress
+		if ((OverFlowAllocators.Num() < OverFlowPoolsCount) && (PoolAllocationStatus == ERHIPoolAllocationStatus::Idle))
+		{
+			PoolAllocationStatus = ERHIPoolAllocationStatus::Allocating;
+
+			CS.Unlock();
+
+			AsyncTask(ENamedThreads::AnyBackgroundHiPriTask, [this]()
+			{
+				NewAllocatorAsyncAllocate();
+			});
+		}
+		else
+		{
+			CS.Unlock();
+		}
+
+		return AllocationSuccess;
 	}
+	else
+	{
+		FScopeLock Lock(&CS);
+
+		for (int32 i = 0; i < Allocators.Num(); i++)
+		{
+			if (Allocators[i]->TryAllocate(SizeInBytes, Alignment, ResourceLocation))
+			{
+				return true;
+			}
+		}
+
+		Allocators.Add(CreateNewAllocator(SizeInBytes));
+		return Allocators.Last()->TryAllocate(SizeInBytes, Alignment, ResourceLocation);
 
-	Allocators.Add(CreateNewAllocator(SizeInBytes));
-	return Allocators.Last()->TryAllocate(SizeInBytes, Alignment, ResourceLocation);
+	}
 }
 
 void FD3D12MultiBuddyAllocator::Deallocate(FD3D12ResourceLocation& ResourceLocation)
@@ -762,6 +881,8 @@ void FD3D12MultiBuddyAllocator::Deallocate(FD3D12ResourceLocation& ResourceLocat
 
 FD3D12BuddyAllocator* FD3D12MultiBuddyAllocator::CreateNewAllocator(uint32 InMinSizeInBytes)
 {
+	TRACE_CPUPROFILER_EVENT_SCOPE(FD3D12MultiBuddyAllocator::CreateNewAllocator);
+
 	check(InMinSizeInBytes <= MaximumAllocationSizeForPooling);
 	uint32 AllocationSize = (InMinSizeInBytes > DefaultPoolSize) ? FMath::RoundUpToPowerOfTwo(InMinSizeInBytes) : DefaultPoolSize;
 
@@ -776,9 +897,49 @@ FD3D12BuddyAllocator* FD3D12MultiBuddyAllocator::CreateNewAllocator(uint32 InMin
 		TraceHeapId);
 }
 
+void FD3D12MultiBuddyAllocator::NewAllocatorAsyncAllocate()
+{
+	TRACE_CPUPROFILER_EVENT_SCOPE(FD3D12MultiBuddyAllocator::NewAllocatorAsyncAllocate);
+
+	FD3D12BuddyAllocator* NewBuddyAllocator = CreateNewAllocator(MaximumAllocationSizeForPooling);
+
+	{
+		FScopeLock Lock(&CS);
+
+		check(PoolAllocationStrategy == ERHIPoolAllocationStrategy::Async);
+		
+		if (PoolAllocationStatus == ERHIPoolAllocationStatus::Allocating)
+		{
+			OverFlowAllocators.Add(NewBuddyAllocator);
+		}
+		else
+		{
+			delete NewBuddyAllocator;
+		}
+
+		PoolAllocationStatus = ERHIPoolAllocationStatus::Idle;
+	}
+}
+
 void FD3D12MultiBuddyAllocator::Initialize()
 {
-	Allocators.Add(CreateNewAllocator(DefaultPoolSize));
+	if (PoolAllocationStrategy == ERHIPoolAllocationStrategy::Async)
+	{
+		Allocators.Add(CreateNewAllocator(MaximumAllocationSizeForPooling));
+
+		check(OverFlowPoolsCount > 0);
+
+		PoolAllocationStatus = ERHIPoolAllocationStatus::Allocating;
+
+		AsyncTask(ENamedThreads::AnyBackgroundHiPriTask, [this]()
+			{
+				NewAllocatorAsyncAllocate();
+			});
+	}
+	else
+	{
+		Allocators.Add(CreateNewAllocator(MaximumAllocationSizeForPooling));
+	}
 }
 
 void FD3D12MultiBuddyAllocator::Destroy()
@@ -831,6 +992,21 @@ void FD3D12MultiBuddyAllocator::UpdateMemoryStats(uint32& IOMemoryAllocated, uin
 
 void FD3D12MultiBuddyAllocator::ReleaseAllResources()
 {
+	if (PoolAllocationStrategy == ERHIPoolAllocationStrategy::Async)
+	{
+		FScopeLock Lock(&CS);
+		PoolAllocationStatus = ERHIPoolAllocationStatus::ShutDown;
+	}
+
+	for (int32 c = 0; c < OverFlowAllocators.Num(); c++)
+	{
+		if (OverFlowAllocators[c] != nullptr)
+		{
+			OverFlowAllocators[c]->Destroy();
+			delete(OverFlowAllocators[c]);
+		}
+	}
+
 	for (int32 i = (Allocators.Num() - 1); i >= 0; i--)
 	{
 		if (Allocators[i])
@@ -840,6 +1016,7 @@ void FD3D12MultiBuddyAllocator::ReleaseAllResources()
 		}
 	}
 
+	OverFlowAllocators.Empty();
 	Allocators.Empty();
 }
 
@@ -866,7 +1043,9 @@ FD3D12UploadHeapAllocator::FD3D12UploadHeapAllocator(FD3D12Adapter* InParent, FD
 		, GD3D12UploadHeapSmallBlockMaxAllocationSize
 		, GD3D12UploadHeapSmallBlockPoolSize
 		, D3D12_TEXTURE_DATA_PLACEMENT_ALIGNMENT
-		, TraceHeapId)
+		, TraceHeapId
+		, GetPoolAllocationStrategy(GD3D12UploadHeapPoolAllocationStrategy)
+		, GD3D12UploadHeapOverFlowPoolsCount)
 
 	, BigBlockAllocator(
 		InParentDevice
@@ -879,7 +1058,9 @@ FD3D12UploadHeapAllocator::FD3D12UploadHeapAllocator(FD3D12Adapter* InParent, FD
 		, GD3D12UploadHeapBigBlockMaxAllocationSize
 		, FRHIMemoryPool::EFreeListOrder::SortByOffset
 		, false /*defrag*/
-		, TraceHeapId)
+		, TraceHeapId
+		, GetPoolAllocationStrategy(GD3D12UploadHeapPoolAllocationStrategy)
+		, GD3D12UploadHeapOverFlowPoolsCount)
 
 	, FastConstantPageAllocator(
 		InParentDevice
@@ -890,10 +1071,17 @@ FD3D12UploadHeapAllocator::FD3D12UploadHeapAllocator(FD3D12Adapter* InParent, FD
 		, GD3D12FastConstantAllocatorPageSize * 64
 		, GD3D12UploadHeapSmallBlockPoolSize
 		, GD3D12FastConstantAllocatorPageSize
-		, TraceHeapId)
+		, TraceHeapId
+		, GetPoolAllocationStrategy(GD3D12UploadHeapPoolAllocationStrategy)
+		, GD3D12UploadHeapOverFlowPoolsCount)
 {
 }
 
+void FD3D12UploadHeapAllocator::Init()
+{
+	SmallBlockAllocator.Initialize();
+	BigBlockAllocator.Initialize(ERHIPoolResourceTypes::Buffers);
+}
 
 void FD3D12UploadHeapAllocator::Destroy()
 {
@@ -1239,6 +1427,9 @@ FD3D12BufferPool* FD3D12DefaultBufferAllocator::CreateBufferPool(D3D12_HEAP_TYPE
 	// Disable defrag if not Default memory
 	bool bDefragEnabled = (InitConfig.HeapType == D3D12_HEAP_TYPE_DEFAULT);
 
+	ERHIPoolAllocationStrategy PoolAllocationStrategy = GetPoolAllocationStrategy(GD3D12DefaultBufferPoolAllocationStrategy);
+	uint16 OverFlowPoolCount = GD3D12DefaultBufferOverFlowPoolsCount;
+
 #if D3D12_RHI_RAYTRACING
 	// Disable defrag on the RT Acceleration pool - #kenzo_todo
 	// Acceleration structure buffers may be created, but not built immediately.
@@ -1250,10 +1441,13 @@ FD3D12BufferPool* FD3D12DefaultBufferAllocator::CreateBufferPool(D3D12_HEAP_TYPE
 		// Use custom pool and allocation size for RT structures because they don't defrag and will thus 'waste' more memory
 		PoolSize = BUFFER_POOL_RT_ACCELERATION_STRUCTURE_POOL_SIZE;
 		MaxAllocationSize = BUFFER_POOL_RT_ACCELERATION_STRUCTURE_MAX_ALLOC_SIZE;
+
+		PoolAllocationStrategy = ERHIPoolAllocationStrategy::Immediate;
+		OverFlowPoolCount = 0;
 	}
 #endif // D3D12_RHI_RAYTRACING
 
-	FD3D12BufferPool* NewPool = new FD3D12PoolAllocator(Device, VisibleNodes, InitConfig, Name, AllocationStrategy, PoolSize, PoolAlignment, MaxAllocationSize, FreeListOrder, bDefragEnabled, TraceHeapId);
+	FD3D12BufferPool* NewPool = new FD3D12PoolAllocator(Device, VisibleNodes, InitConfig, Name, AllocationStrategy, PoolSize, PoolAlignment, MaxAllocationSize, FreeListOrder, bDefragEnabled, TraceHeapId, PoolAllocationStrategy, OverFlowPoolCount);
 
 #else // USE_BUFFER_POOL_ALLOCATOR
 
@@ -1271,7 +1465,9 @@ FD3D12BufferPool* FD3D12DefaultBufferAllocator::CreateBufferPool(D3D12_HEAP_TYPE
 		InHeapType == D3D12_HEAP_TYPE_READBACK ? READBACK_BUFFER_POOL_MAX_ALLOC_SIZE : DEFAULT_BUFFER_POOL_MAX_ALLOC_SIZE,
 		InHeapType == D3D12_HEAP_TYPE_READBACK ? READBACK_BUFFER_POOL_DEFAULT_POOL_SIZE : DEFAULT_BUFFER_POOL_DEFAULT_POOL_SIZE,
 		MinBlockSize,
-		TraceHeapId
+		TraceHeapId,
+		ERHIPoolAllocationStrategy::Immediate,
+		0
 		);
 
 	FD3D12DefaultBufferPool* NewPool = new FD3D12DefaultBufferPool(Device, Allocator);
@@ -1571,8 +1767,8 @@ FD3D12TextureAllocatorPool::FD3D12TextureAllocatorPool(FD3D12Device* Device, FRH
 		uint64 PoolSize = 4 * 1024 * 1024;
 		uint64 MaxAllocationSize = PoolSize;
 		bool bDefragEnabled = false; // Disable defrag on 4K pool because it shouldn't really fragment - all allocations are 4K or multiple of 4K and pretty small
-		FD3D12PoolAllocator* ReadOnly4KPool = new FD3D12PoolAllocator(Device, GetVisibilityMask(), InitConfig, Name, AllocationStrategy, PoolSize, D3D12_SMALL_RESOURCE_PLACEMENT_ALIGNMENT, MaxAllocationSize, FreeListOrder, bDefragEnabled, TraceHeapId);
-		ReadOnly4KPool->Initialize();
+		FD3D12PoolAllocator* ReadOnly4KPool = new FD3D12PoolAllocator(Device, GetVisibilityMask(), InitConfig, Name, AllocationStrategy, PoolSize, D3D12_SMALL_RESOURCE_PLACEMENT_ALIGNMENT, MaxAllocationSize, FreeListOrder, bDefragEnabled, TraceHeapId, ERHIPoolAllocationStrategy::Immediate, 0);
+		ReadOnly4KPool->Initialize(ERHIPoolResourceTypes::NonRTDSTextures);
 
 		PoolAllocators[(int)EPoolType::ReadOnly4K] = ReadOnly4KPool;
 	}
@@ -1585,8 +1781,8 @@ FD3D12TextureAllocatorPool::FD3D12TextureAllocatorPool(FD3D12Device* Device, FRH
 		uint64 PoolSize = GD3D12PoolAllocatorReadOnlyTextureVRAMPoolSize;
 		uint64 MaxAllocationSize = GD3D12PoolAllocatorReadOnlyTextureVRAMMaxAllocationSize;
 		bool bDefragEnabled = true;
-		FD3D12PoolAllocator* ReadOnlyPool = new FD3D12PoolAllocator(Device, GetVisibilityMask(), InitConfig, Name, AllocationStrategy, PoolSize, D3D12_DEFAULT_RESOURCE_PLACEMENT_ALIGNMENT, MaxAllocationSize, FreeListOrder, bDefragEnabled, TraceHeapId);
-		ReadOnlyPool->Initialize();
+		FD3D12PoolAllocator* ReadOnlyPool = new FD3D12PoolAllocator(Device, GetVisibilityMask(), InitConfig, Name, AllocationStrategy, PoolSize, D3D12_DEFAULT_RESOURCE_PLACEMENT_ALIGNMENT, MaxAllocationSize, FreeListOrder, bDefragEnabled, TraceHeapId, ERHIPoolAllocationStrategy::Immediate, 0);
+		ReadOnlyPool->Initialize(ERHIPoolResourceTypes::NonRTDSTextures);
 
 		PoolAllocators[(int)EPoolType::ReadOnly] = ReadOnlyPool;
 	}
@@ -1600,8 +1796,8 @@ FD3D12TextureAllocatorPool::FD3D12TextureAllocatorPool(FD3D12Device* Device, FRH
 		uint64 MaxAllocationSize = GD3D12PoolAllocatorRTUAVTextureVRAMMaxAllocationSize;
 		// FD3D12ResourceLocation::OnAllocationMoved doesn't correctly retrieve the clear value when recreating moved resources, so we need to disable defrag for this pool for the time being.
 		bool bDefragEnabled = false;
-		FD3D12PoolAllocator* RTPool = new FD3D12PoolAllocator(Device, GetVisibilityMask(), InitConfig, Name, AllocationStrategy, PoolSize, D3D12_DEFAULT_RESOURCE_PLACEMENT_ALIGNMENT, MaxAllocationSize, FreeListOrder, bDefragEnabled, TraceHeapId);
-		RTPool->Initialize();
+		FD3D12PoolAllocator* RTPool = new FD3D12PoolAllocator(Device, GetVisibilityMask(), InitConfig, Name, AllocationStrategy, PoolSize, D3D12_DEFAULT_RESOURCE_PLACEMENT_ALIGNMENT, MaxAllocationSize, FreeListOrder, bDefragEnabled, TraceHeapId, ERHIPoolAllocationStrategy::Immediate, 0);
+		RTPool->Initialize(ERHIPoolResourceTypes::NonRTDSTextures);
 
 		PoolAllocators[(int)EPoolType::RenderTarget] = RTPool;
 	}
@@ -1615,8 +1811,8 @@ FD3D12TextureAllocatorPool::FD3D12TextureAllocatorPool(FD3D12Device* Device, FRH
 		uint64 MaxAllocationSize = GD3D12PoolAllocatorRTUAVTextureVRAMMaxAllocationSize;
 		// Defrag doesn't correctly handle resources which need the BCn/UINT UAV aliasing workaround, so we'll turn off defrag for this heap for now.
 		bool bDefragEnabled = false;
-		FD3D12PoolAllocator* UAVPool = new FD3D12PoolAllocator(Device, GetVisibilityMask(), InitConfig, Name, AllocationStrategy, PoolSize, D3D12_DEFAULT_RESOURCE_PLACEMENT_ALIGNMENT, MaxAllocationSize, FreeListOrder, bDefragEnabled, TraceHeapId);
-		UAVPool->Initialize();
+		FD3D12PoolAllocator* UAVPool = new FD3D12PoolAllocator(Device, GetVisibilityMask(), InitConfig, Name, AllocationStrategy, PoolSize, D3D12_DEFAULT_RESOURCE_PLACEMENT_ALIGNMENT, MaxAllocationSize, FreeListOrder, bDefragEnabled, TraceHeapId, ERHIPoolAllocationStrategy::Immediate, 0);
+		UAVPool->Initialize(ERHIPoolResourceTypes::NonRTDSTextures);
 
 		PoolAllocators[(int)EPoolType::UAV] = UAVPool;
 	}
@@ -1750,7 +1946,9 @@ FD3D12TextureAllocator::FD3D12TextureAllocator(FD3D12Device* Device,
 		D3D12_DEFAULT_RESOURCE_PLACEMENT_ALIGNMENT,
 		HeapSize,
 		D3D12_SMALL_RESOURCE_PLACEMENT_ALIGNMENT,
-		InTraceParentHeapId)
+		InTraceParentHeapId,
+		ERHIPoolAllocationStrategy::Immediate, 
+		0)
 {
 }
 
diff --git a/Engine/Source/Runtime/D3D12RHI/Private/D3D12Allocation.h b/Engine/Source/Runtime/D3D12RHI/Private/D3D12Allocation.h
--- a/Engine/Source/Runtime/D3D12RHI/Private/D3D12Allocation.h
+++ b/Engine/Source/Runtime/D3D12RHI/Private/D3D12Allocation.h
@@ -206,7 +206,9 @@ public:
 		uint32 InMaxAllocationSize,
 		uint32 InDefaultPoolSize,
 		uint32 InMinBlockSize,
-		HeapId InTraceParentHeapId);
+		HeapId InTraceParentHeapId,
+		ERHIPoolAllocationStrategy InPoolAllocationStrategy,
+		uint16 InOverFlowPoolsCount);
 	~FD3D12MultiBuddyAllocator();
 
 	const EResourceAllocationStrategy GetAllocationStrategy() const { return AllocationStrategy; }
@@ -229,8 +231,9 @@ public:
 	void Reset();
 	
 protected:
-
 	FD3D12BuddyAllocator* CreateNewAllocator(uint32 InMinSizeInBytes);
+	void NewAllocatorAsyncAllocate();
+
 
 	const EResourceAllocationStrategy AllocationStrategy;
 	const uint32 MinBlockSize;
@@ -240,6 +243,11 @@ protected:
 
 private:
 	HeapId TraceHeapId;
+
+	TArray<FD3D12BuddyAllocator*> OverFlowAllocators;
+	ERHIPoolAllocationStrategy PoolAllocationStrategy = ERHIPoolAllocationStrategy::Immediate;
+	ERHIPoolAllocationStatus PoolAllocationStatus = ERHIPoolAllocationStatus::Idle;
+	int16 OverFlowPoolsCount = 1;
 };
 
 //-----------------------------------------------------------------------------
@@ -253,7 +261,7 @@ public:
 
 	FD3D12UploadHeapAllocator(FD3D12Adapter* InParent, FD3D12Device* InParentDevice, const FString& InName);
 
-	void Init() {}
+	void Init();
 	void Destroy();
 
 	// Allocates <size> bytes from the end of an available resource heap.
@@ -270,7 +278,7 @@ private:
 	// Pool allocator for all bigger allocations - less fast but less alignment waste
 	FCriticalSection			BigBlockCS;
 	FD3D12PoolAllocator			BigBlockAllocator;
-	// Seperate buddy allocator used for the fast constant allocator pages which get always freed within the same frame by default
+	// Separate buddy allocator used for the fast constant allocator pages which get always freed within the same frame by default
 	// (different allocator to avoid fragmentation with the other pools - always the same size allocations)
 	FD3D12MultiBuddyAllocator	FastConstantPageAllocator;
 };
diff --git a/Engine/Source/Runtime/D3D12RHI/Private/D3D12PoolAllocator.cpp b/Engine/Source/Runtime/D3D12RHI/Private/D3D12PoolAllocator.cpp
--- a/Engine/Source/Runtime/D3D12RHI/Private/D3D12PoolAllocator.cpp
+++ b/Engine/Source/Runtime/D3D12RHI/Private/D3D12PoolAllocator.cpp
@@ -239,8 +239,8 @@ EResourceAllocationStrategy FD3D12PoolAllocator::GetResourceAllocationStrategy(D
 
 
 FD3D12PoolAllocator::FD3D12PoolAllocator(FD3D12Device* ParentDevice, FRHIGPUMask VisibleNodes, const FD3D12ResourceInitConfig& InInitConfig, const FString& InName,
-	EResourceAllocationStrategy InAllocationStrategy, uint64 InDefaultPoolSize, uint32 InPoolAlignment, uint32 InMaxAllocationSize, FRHIMemoryPool::EFreeListOrder InFreeListOrder, bool bInDefragEnabled, HeapId InTraceParentHeapId) :
-	FRHIPoolAllocator(InDefaultPoolSize, InPoolAlignment, InMaxAllocationSize, InFreeListOrder, bInDefragEnabled), 
+	EResourceAllocationStrategy InAllocationStrategy, uint64 InDefaultPoolSize, uint32 InPoolAlignment, uint32 InMaxAllocationSize, FRHIMemoryPool::EFreeListOrder InFreeListOrder, bool bInDefragEnabled, HeapId InTraceParentHeapId, ERHIPoolAllocationStrategy InPoolAllocationStrategy, uint16 InOverFlowPoolsCount) :
+	FRHIPoolAllocator(InDefaultPoolSize, InPoolAlignment, InMaxAllocationSize, InFreeListOrder, bInDefragEnabled, InPoolAllocationStrategy, InOverFlowPoolsCount),
 	FD3D12DeviceChild(ParentDevice), 
 	FD3D12MultiNodeGPUObject(ParentDevice->GetGPUMask(), VisibleNodes), 
 	InitConfig(InInitConfig), 
diff --git a/Engine/Source/Runtime/D3D12RHI/Private/D3D12PoolAllocator.h b/Engine/Source/Runtime/D3D12RHI/Private/D3D12PoolAllocator.h
--- a/Engine/Source/Runtime/D3D12RHI/Private/D3D12PoolAllocator.h
+++ b/Engine/Source/Runtime/D3D12RHI/Private/D3D12PoolAllocator.h
@@ -132,7 +132,7 @@ public:
 
 	// Constructor
 	FD3D12PoolAllocator(FD3D12Device* ParentDevice, FRHIGPUMask VisibleNodes, const FD3D12ResourceInitConfig& InInitConfig, const FString& InName,
-		EResourceAllocationStrategy InAllocationStrategy, uint64 InDefaultPoolSize, uint32 InPoolAlignment, uint32 InMaxAllocationSize, FRHIMemoryPool::EFreeListOrder InFreeListOrder, bool bInDefragEnabled, HeapId InTraceParentHeapId);
+		EResourceAllocationStrategy InAllocationStrategy, uint64 InDefaultPoolSize, uint32 InPoolAlignment, uint32 InMaxAllocationSize, FRHIMemoryPool::EFreeListOrder InFreeListOrder, bool bInDefragEnabled, HeapId InTraceParentHeapId, ERHIPoolAllocationStrategy InPoolAllocationStrategy, uint16 InOverFlowPoolsCount);
 	~FD3D12PoolAllocator();
 	
 	// Function names currently have to match with FD3D12DefaultBufferPool until we can make full replacement of the allocator
diff --git a/Engine/Source/Runtime/RHICore/Private/RHIPoolAllocator.cpp b/Engine/Source/Runtime/RHICore/Private/RHIPoolAllocator.cpp
--- a/Engine/Source/Runtime/RHICore/Private/RHIPoolAllocator.cpp
+++ b/Engine/Source/Runtime/RHICore/Private/RHIPoolAllocator.cpp
@@ -2,6 +2,7 @@
 
 #include "RHIPoolAllocator.h"
 #include "HAL/IConsoleManager.h"
+#include "Async/Async.h"
 
 static float GRHIPoolAllocatorDefragSizeFraction = 0.9f;
 static FAutoConsoleVariableRef CVarRHIPoolAllocatorDefragSizeFraction(
@@ -320,7 +321,17 @@ void FRHIMemoryPool::Init()
 		AllocationDataPool.Add(AllocationData);
 	}
 
+	// Async allocated pool will not have a proper index at Init time
+	if (PoolIndex != InvalidPoolIndex)
+	{
+		InitFreeBlocks();
+	}
+}
+
+void FRHIMemoryPool::InitFreeBlocks()
+{
 	// Setup the special start and free block
+	check(PoolIndex != InvalidPoolIndex);
 	HeadBlock.InitAsHead(PoolIndex);
 	FRHIPoolAllocationData* FreeBlock = GetNewAllocationData();
 	FreeBlock->InitAsFree(PoolIndex, PoolSize, PoolAlignment, 0);
@@ -593,6 +604,9 @@ void FRHIMemoryPool::Validate()
 	if (GRHIPoolAllocatorValidateLinkList <= 0)
 		return;
 
+	if (PoolIndex == InvalidPoolIndex)
+		return;
+
 	uint64 TotalFreeSize = 0;
 	uint64 TotalAllocatedSize = 0;
 	uint64 CurrentOffset = 0;
@@ -650,14 +664,17 @@ void FRHIMemoryPool::Validate()
 //	Pool Allocator
 //-----------------------------------------------------------------------------
 
-FRHIPoolAllocator::FRHIPoolAllocator(uint64 InDefaultPoolSize, uint32 InPoolAlignment, uint32 InMaxAllocationSize, FRHIMemoryPool::EFreeListOrder InFreeListOrder, bool InDefragEnabled) :
+FRHIPoolAllocator::FRHIPoolAllocator(uint64 InDefaultPoolSize, uint32 InPoolAlignment, uint32 InMaxAllocationSize, FRHIMemoryPool::EFreeListOrder InFreeListOrder, 
+	bool InDefragEnabled, ERHIPoolAllocationStrategy InPoolAllocationStrategy, uint16 InOverFlowPoolsCount):
 	DefaultPoolSize(InDefaultPoolSize),
 	PoolAlignment(InPoolAlignment),
 	MaxAllocationSize(InMaxAllocationSize),
 	FreeListOrder(InFreeListOrder),
 	bDefragEnabled(InDefragEnabled),
 	LastDefragPoolIndex(0),
-	TotalAllocatedBlocks(0)
+	TotalAllocatedBlocks(0),
+	PoolAllocationStrategy(InPoolAllocationStrategy),
+	OverFlowPoolsCount(InOverFlowPoolsCount)
 {}
 
 
@@ -667,13 +684,35 @@ FRHIPoolAllocator::~FRHIPoolAllocator()
 }
 
 
-void FRHIPoolAllocator::Initialize()
-{	
+void FRHIPoolAllocator::Initialize(ERHIPoolResourceTypes InAsyncAllocationResourceType)
+{
+	if (PoolAllocationStrategy == ERHIPoolAllocationStrategy::Async)
+	{
+		check(Pools.Num() == 0);
+
+		Pools.Add(CreateNewPool(0, 16, InAsyncAllocationResourceType));
+		PoolAllocationOrder.Add(0);
+
+		check(OverFlowPoolsCount > 0);
+
+		PoolAllocationStatus = ERHIPoolAllocationStatus::Allocating;
+
+		AsyncTask(ENamedThreads::AnyBackgroundHiPriTask, [this, InAsyncAllocationResourceType]()
+			{
+				PoolAsyncAllocate(DefaultPoolSize, 16, InAsyncAllocationResourceType);
+			});
+	}
 }
 
 
 void FRHIPoolAllocator::Destroy()
 {
+	if (PoolAllocationStrategy == ERHIPoolAllocationStrategy::Async)
+	{
+		FScopeLock Lock(&CS);
+		PoolAllocationStatus = ERHIPoolAllocationStatus::ShutDown;
+	}
+
 	for (int32 i = (Pools.Num() - 1); i >= 0; i--)
 	{
 		if (Pools[i] != nullptr)
@@ -683,8 +722,18 @@ void FRHIPoolAllocator::Destroy()
 		}
 	}
 
+	for (int32 c = 0; c < OverFlowPools.Num(); c++)
+	{
+		if (OverFlowPools[c] != nullptr)
+		{
+			OverFlowPools[c]->Destroy();
+			delete(OverFlowPools[c]);
+		}
+	}
+	
 	Pools.Empty();
 	PoolAllocationOrder.Empty();
+	OverFlowPools.Empty();
 }
 
 
@@ -692,42 +741,181 @@ bool FRHIPoolAllocator::TryAllocateInternal(uint32 InSizeInBytes, uint32 InAlloc
 {
 	TRACE_CPUPROFILER_EVENT_SCOPE(FRHIPoolAllocator::TryAllocateInternal);
 
-	FScopeLock Lock(&CS);
-
-	check(PoolAllocationOrder.Num() == Pools.Num())
-	for (int32 PoolIndex : PoolAllocationOrder)
+	if (PoolAllocationStrategy == ERHIPoolAllocationStrategy::Async)
 	{
-		FRHIMemoryPool* Pool = Pools[PoolIndex];
-		if (Pool != nullptr && Pool->IsResourceTypeSupported(InAllocationResourceType) && !Pool->IsFull() && Pool->TryAllocate(InSizeInBytes, InAllocationAlignment, InAllocationResourceType, AllocationData))
+		CS.Lock();
+
+		check(PoolAllocationOrder.Num() == Pools.Num())
+
+		bool bResourceAllocated = false;
+
+		for (int32 PoolIndex : PoolAllocationOrder)
 		{
-			return true;
+			FRHIMemoryPool* Pool = Pools[PoolIndex];
+			if (Pool != nullptr && Pool->IsResourceTypeSupported(InAllocationResourceType) && !Pool->IsFull() && Pool->TryAllocate(InSizeInBytes, InAllocationAlignment, InAllocationResourceType, AllocationData))
+			{
+				bResourceAllocated = true;
+				break;
+			}
 		}
-	}
+		uint32 AlignedSize = 0;
+		// No pool available, grab an overflow pool or we need to allocate a new pool right now
+		if (!bResourceAllocated)
+		{ 
+			// Find a free pool index
+			int16 NewPoolIndex = -1;
+			for (int32 PoolIndex = 0; PoolIndex < Pools.Num(); PoolIndex++)
+			{
+				if (Pools[PoolIndex] == nullptr)
+				{
+					NewPoolIndex = PoolIndex;
+					break;
+				}
+			}
 
-	// Find a free pool index
-	int16 NewPoolIndex = -1;
-	for (int32 PoolIndex = 0; PoolIndex < Pools.Num(); PoolIndex++)
-	{
-		if (Pools[PoolIndex] == nullptr)
+			int16 PoolIndexForInit = (NewPoolIndex >= 0) ? NewPoolIndex : Pools.Num();
+
+			AlignedSize = FRHIMemoryPool::GetAlignedSize(InSizeInBytes, PoolAlignment, InAllocationAlignment);
+
+			FRHIMemoryPool* ReadyPool = nullptr;
+			for (int32 OverFlowPoolIndex = 0; OverFlowPoolIndex < OverFlowPools.Num(); OverFlowPoolIndex++)
+			{
+				if ((OverFlowPools[OverFlowPoolIndex]->IsResourceTypeSupported(InAllocationResourceType)) 
+					&& (OverFlowPools[OverFlowPoolIndex]->GetPoolAlignment() == PoolAlignment)
+					&& (OverFlowPools[OverFlowPoolIndex]->GetPoolSize() >= AlignedSize))
+				{
+					ReadyPool = OverFlowPools[OverFlowPoolIndex];
+
+					check(PoolIndexForInit != -1);
+					
+					// Pool needs correct PoolIndex for FreeBlocks
+					ReadyPool->SetPoolIndex(PoolIndexForInit);
+					ReadyPool->InitFreeBlocks();
+
+					check(OverFlowPools[OverFlowPoolIndex]->IsEmpty());
+					OverFlowPools.RemoveAtSwap(OverFlowPoolIndex);
+					break;
+				}
+			}
+
+			// No overflow pools ready, need to allocate inline a new pool
+			if (ReadyPool == nullptr)
+			{
+				ReadyPool = CreateNewPool(PoolIndexForInit, AlignedSize, InAllocationResourceType);
+			}
+
+			if (NewPoolIndex >= 0)
+			{
+				Pools[NewPoolIndex] = ReadyPool;
+			}
+			else
+			{
+				NewPoolIndex = Pools.Num();
+				Pools.Add(ReadyPool);
+				PoolAllocationOrder.Add(NewPoolIndex);
+			}
+
+			check(NewPoolIndex == ReadyPool->GetPoolIndex());
+
+			bResourceAllocated = Pools[NewPoolIndex]->TryAllocate(InSizeInBytes, InAllocationAlignment, InAllocationResourceType, AllocationData);
+		}
+
+		// Test if we have enough overflow pools available
+		int32 FreePoolsCount = 0;
+
+		for (int32 OverFlowPoolIndex = 0; OverFlowPoolIndex < OverFlowPools.Num(); OverFlowPoolIndex++)
 		{
-			NewPoolIndex = PoolIndex;
-			break;
+			if ((OverFlowPools[OverFlowPoolIndex]->IsResourceTypeSupported(InAllocationResourceType))
+				&& (OverFlowPools[OverFlowPoolIndex]->GetPoolAlignment() == PoolAlignment)
+				&& (OverFlowPools[OverFlowPoolIndex]->GetPoolSize() >= AlignedSize))
+			{
+				FreePoolsCount++;
+			}
 		}
-	}
 
-	uint32 AlignedSize = FRHIMemoryPool::GetAlignedSize(InSizeInBytes, PoolAlignment, InAllocationAlignment);
-	if (NewPoolIndex >= 0)
-	{
-		Pools[NewPoolIndex] = CreateNewPool(NewPoolIndex, AlignedSize, InAllocationResourceType);
+		if ((FreePoolsCount < OverFlowPoolsCount) && (PoolAllocationStatus == ERHIPoolAllocationStatus::Idle))
+		{
+			PoolAllocationStatus = ERHIPoolAllocationStatus::Allocating;
+
+			CS.Unlock();
+
+			AsyncTask(ENamedThreads::AnyBackgroundHiPriTask, [this, /*DefaultPoolSize, PoolAlignment,*/ InAllocationResourceType]()
+			{
+				PoolAsyncAllocate(DefaultPoolSize, PoolAlignment, InAllocationResourceType);
+			});
+		}
+		else
+		{
+			CS.Unlock();
+		}
+
+		return bResourceAllocated;
 	}
 	else
 	{
-		NewPoolIndex = Pools.Num();		
-		Pools.Add(CreateNewPool(NewPoolIndex, AlignedSize, InAllocationResourceType));
-		PoolAllocationOrder.Add(NewPoolIndex);
+		FScopeLock Lock(&CS);
+
+		check(PoolAllocationOrder.Num() == Pools.Num())
+
+		for (int32 PoolIndex : PoolAllocationOrder)
+		{
+			FRHIMemoryPool* Pool = Pools[PoolIndex];
+			if (Pool != nullptr && Pool->IsResourceTypeSupported(InAllocationResourceType) && !Pool->IsFull() && Pool->TryAllocate(InSizeInBytes, InAllocationAlignment, InAllocationResourceType, AllocationData))
+			{
+				return true;
+			}
+		}
+
+		// Find a free pool index
+		int16 NewPoolIndex = -1;
+		for (int32 PoolIndex = 0; PoolIndex < Pools.Num(); PoolIndex++)
+		{
+			if (Pools[PoolIndex] == nullptr)
+			{
+				NewPoolIndex = PoolIndex;
+				break;
+			}
+		}
+
+		uint32 AlignedSize = FRHIMemoryPool::GetAlignedSize(InSizeInBytes, PoolAlignment, InAllocationAlignment);
+		if (NewPoolIndex >= 0)
+		{
+			Pools[NewPoolIndex] = CreateNewPool(NewPoolIndex, AlignedSize, InAllocationResourceType);
+		}
+		else
+		{
+			NewPoolIndex = Pools.Num();
+			Pools.Add(CreateNewPool(NewPoolIndex, AlignedSize, InAllocationResourceType));
+			PoolAllocationOrder.Add(NewPoolIndex);
+		}
+
+		return Pools[NewPoolIndex]->TryAllocate(InSizeInBytes, InAllocationAlignment, InAllocationResourceType, AllocationData);
 	}
+}
+
+void FRHIPoolAllocator::PoolAsyncAllocate(uint32 InSizeInBytes, uint32 InAllocationAlignment, ERHIPoolResourceTypes InAllocationResourceType)
+{
+	TRACE_CPUPROFILER_EVENT_SCOPE(FRHIPoolAllocator::PoolAsyncAllocate);
+
+	uint32 AlignedSize = FRHIMemoryPool::GetAlignedSize(InSizeInBytes, PoolAlignment, InAllocationAlignment);
+	FRHIMemoryPool* NewPool = CreateNewPool(FRHIMemoryPool::InvalidPoolIndex, AlignedSize, InAllocationResourceType);
+
+	{
+		FScopeLock Lock(&CS);
+		check(PoolAllocationStrategy == ERHIPoolAllocationStrategy::Async);
+		
+		if (PoolAllocationStatus == ERHIPoolAllocationStatus::Allocating)
+		{
+			// Allocate the pool without knowing the correct Index, it will be initialized later when used for allocations
+			OverFlowPools.Add(NewPool);
+		}
+		else
+		{
+			delete NewPool;
+		}
 
-	return Pools[NewPoolIndex]->TryAllocate(InSizeInBytes, InAllocationAlignment, InAllocationResourceType, AllocationData);
+		PoolAllocationStatus = ERHIPoolAllocationStatus::Idle;
+	}
 }
 
 
diff --git a/Engine/Source/Runtime/RHICore/Public/RHIPoolAllocator.h b/Engine/Source/Runtime/RHICore/Public/RHIPoolAllocator.h
--- a/Engine/Source/Runtime/RHICore/Public/RHIPoolAllocator.h
+++ b/Engine/Source/Runtime/RHICore/Public/RHIPoolAllocator.h
@@ -31,6 +31,26 @@ enum class ERHIPoolResourceTypes
 	All				= Buffers | RTDSTextures | NonRTDSTextures,
 };
 
+/**
+@brief Type of allocation strategy
+**/
+enum class ERHIPoolAllocationStrategy
+{
+	Immediate,
+	Async,
+	Count
+};
+
+/**
+@brief Status of the Async pool allocation
+**/
+enum class ERHIPoolAllocationStatus
+{
+	Idle,
+	Allocating,
+	ShutDown,
+	Count
+};
 
 /**
 @brief Pool allocator internal data
@@ -135,6 +155,8 @@ class FRHIMemoryPool
 {
 public:
 
+	static constexpr int16 InvalidPoolIndex = -1;
+
 	enum class EFreeListOrder
 	{
 		SortBySize,
@@ -148,6 +170,7 @@ public:
 	// Setup/Shutdown
 	RHICORE_API virtual void Init();
 	RHICORE_API virtual void Destroy();
+	RHICORE_API void InitFreeBlocks();
 
 	// AllocationData functions: allocate, free, move, ...
 	RHICORE_API bool TryAllocate(uint32 InSizeInBytes, uint32 InAllocationAlignment, ERHIPoolResourceTypes InAllocationResourceType, FRHIPoolAllocationData& AllocationData);
@@ -161,6 +184,7 @@ public:
 	uint64 GetPoolSize() const { return PoolSize; }
 	uint64 GetFreeSize() const { return FreeSize; }
 	uint64 GetUsedSize() const { return (PoolSize - FreeSize); }
+	uint32 GetPoolAlignment() const { return PoolAlignment; }
 	uint32 GetAlignmentWaste() const { return AligmnentWaste; }
 	uint32 GetAllocatedBlocks() const { return AllocatedBlocks; }
 	ERHIPoolResourceTypes GetSupportedResourceTypes() const { return SupportedResourceTypes; }
@@ -168,6 +192,9 @@ public:
 	bool IsEmpty() const { return GetUsedSize() == 0; }
 	bool IsFull() const { return FreeSize == 0; }
 
+	// Index setter for async allocation
+	void SetPoolIndex(int16 InPoolIndex) { PoolIndex = InPoolIndex; }
+
 	// Static public helper functions
 	static RHICORE_API uint32 GetAlignedSize(uint32 InSizeInBytes, uint32 InPoolAlignment, uint32 InAllocationAlignment);
 	static RHICORE_API uint32 GetAlignedOffset(uint32 InOffset, uint32 InPoolAlignment, uint32 InAllocationAlignment);
@@ -221,11 +248,11 @@ class FRHIPoolAllocator
 public:
 
 	// Constructor
-	RHICORE_API FRHIPoolAllocator(uint64 InDefaultPoolSize, uint32 InPoolAlignment, uint32 InMaxAllocationSize, FRHIMemoryPool::EFreeListOrder InFreeListOrder, bool InDefragEnabled);
+	RHICORE_API FRHIPoolAllocator(uint64 InDefaultPoolSize, uint32 InPoolAlignment, uint32 InMaxAllocationSize, FRHIMemoryPool::EFreeListOrder InFreeListOrder, bool InDefragEnabled, ERHIPoolAllocationStrategy InPoolAllocationStrategy, uint16 InOverFlowPoolsCount);
 	RHICORE_API virtual ~FRHIPoolAllocator();
 
 	// Setup/Shutdown
-	RHICORE_API void Initialize();
+	RHICORE_API void Initialize(ERHIPoolResourceTypes InAllocationResourceType);
 	RHICORE_API void Destroy();
 
 	// Defrag & cleanup operation
@@ -241,6 +268,8 @@ protected:
 	RHICORE_API void DeallocateInternal(FRHIPoolAllocationData& AllocationData);
 
 	RHICORE_API void SortPools();
+
+	RHICORE_API void PoolAsyncAllocate(uint32 InSizeInBytes, uint32 InAllocationAlignment, ERHIPoolResourceTypes InAllocationResourceType);
 	
 	// Helper function to create a new platform specific pool
 	virtual FRHIMemoryPool* CreateNewPool(int16 InPoolIndex, uint32 InMinimumAllocationSize, ERHIPoolResourceTypes InAllocationResourceType) = 0;
@@ -271,6 +300,18 @@ protected:
 
 	// Stats
 	uint32 TotalAllocatedBlocks;
+
+	// Array of pre-allocated pools
+	TArray<FRHIMemoryPool*> OverFlowPools;
+
+	// Pool allocation strategy
+	ERHIPoolAllocationStrategy	PoolAllocationStrategy = ERHIPoolAllocationStrategy::Immediate;
+
+	// Pool allocation status
+	ERHIPoolAllocationStatus	PoolAllocationStatus = ERHIPoolAllocationStatus::Idle;
+
+	// Max pre-allocated pool
+	int16 OverFlowPoolsCount	= 1;
 };
 
 
