Move RHIFrameInfoSetPresentCounter from the RenderThread to the RHIThread since GRHIPresentCounter is incremented there.

#Synced-CL 50646776

diff --git a/Engine/Source/Runtime/Apple/MetalRHI/Private/MetalViewport.cpp b/Engine/Source/Runtime/Apple/MetalRHI/Private/MetalViewport.cpp
--- a/Engine/Source/Runtime/Apple/MetalRHI/Private/MetalViewport.cpp
+++ b/Engine/Source/Runtime/Apple/MetalRHI/Private/MetalViewport.cpp
@@ -1353,6 +1353,9 @@ void FMetalDynamicRHI::RHIEndDrawingViewport(FRHICommandListImmediate& RHICmdLis
 		FRHICommandListScopedPipeline PipelineScope(RHICmdList, ERHIPipeline::Graphics);
 	    RHICmdList.EnqueueLambda(TEXT("RHIEndDrawingViewport"), [ViewportRHI, PresentArgs](FRHICommandListImmediate& ExecutingCmdList)
 	    {
+			// Not all platforms have control on the Present Counter, correlate the information FrameCounter -> PresentCounter before Presenting the frame
+			RHIFrameInfoSetPresentCounter(PresentArgs.FrameCounter, GRHIPresentCounter);
+
 			FMetalRHICommandContext& Context = FMetalRHICommandContext::Get(ExecutingCmdList);
 			Context.EndDrawingViewport(ResourceCast(ViewportRHI), PresentArgs);
 	    });
diff --git a/Engine/Source/Runtime/D3D12RHI/Private/D3D12Viewport.cpp b/Engine/Source/Runtime/D3D12RHI/Private/D3D12Viewport.cpp
--- a/Engine/Source/Runtime/D3D12RHI/Private/D3D12Viewport.cpp
+++ b/Engine/Source/Runtime/D3D12RHI/Private/D3D12Viewport.cpp
@@ -714,6 +714,8 @@ bool FD3D12Viewport::EndRendering(FD3D12CommandContextBase& ContextBase, const F
 		<< End_Rendering.FrameCounter(InPresentArgs.FrameCounter);
 #endif // CPUPROFILERTRACE_ENABLED
 
+	TRACE_BOOKMARK(TEXT("Present P:%" UINT64_FMT " F:%" UINT64_FMT), GRHIPresentCounter, InPresentArgs.FrameCounter)
+
 	check(ContextBase.GetParentAdapter() == GetParentAdapter());
 
 	const bool bCustomPresent = IsValidRef(CustomPresent);
@@ -906,6 +908,9 @@ void FD3D12DynamicRHI::RHIEndDrawingViewport(FRHICommandListImmediate& RHICmdLis
 		{
 			FD3D12CommandContextBase& Context = FD3D12CommandContextBase::Get(ExecutingCmdList);
 
+			// Not all platforms have control on the Present Counter, correlate the information FrameCounter -> PresentCounter before Presenting the frame
+			RHIFrameInfoSetPresentCounter(PresentArgs.FrameCounter, GRHIPresentCounter);
+
 			FD3D12Viewport* Viewport = FD3D12DynamicRHI::ResourceCast(ViewportRHI);
 			SCOPE_CYCLE_COUNTER(STAT_D3D12PresentTime);
 
diff --git a/Engine/Source/Runtime/RHI/Private/RHICommandList.cpp b/Engine/Source/Runtime/RHI/Private/RHICommandList.cpp
--- a/Engine/Source/Runtime/RHI/Private/RHICommandList.cpp
+++ b/Engine/Source/Runtime/RHI/Private/RHICommandList.cpp
@@ -1864,9 +1864,6 @@ void FRHICommandListImmediate::EndDrawingViewport(FRHIViewport* Viewport, const
 {
 	RHI_BREADCRUMB_EVENT(*this, "EndDrawingViewport");
 	
-	// Not all platforms have control on the Present Counter, correlate the information FrameCounter -> PresentCounter before Presenting the frame
-	RHIFrameInfoSetPresentCounter(InPresentArgs.FrameCounter, GRHIPresentCounter);
-	
 	GDynamicRHI->RHIEndDrawingViewport(*this, Viewport, InPresentArgs);
 
 	// Dispatch to ensure any swapchain present commands enqueued by the platform RHI above are executed in a timely manner.
diff --git a/Engine/Source/Runtime/VulkanRHI/Private/VulkanRHI.cpp b/Engine/Source/Runtime/VulkanRHI/Private/VulkanRHI.cpp
--- a/Engine/Source/Runtime/VulkanRHI/Private/VulkanRHI.cpp
+++ b/Engine/Source/Runtime/VulkanRHI/Private/VulkanRHI.cpp
@@ -1268,6 +1268,10 @@ void FVulkanDynamicRHI::RHIEndDrawingViewport(FRHICommandListImmediate& RHICmdLi
 		RHICmdList.EnqueueLambda(TEXT("RHIEndDrawingViewport"), [ViewportRHI, PresentArgs](FRHICommandListImmediate& ExecutingCmdList)
 		{
 			LLM_SCOPE_VULKAN(ELLMTagVulkan::VulkanMisc);
+
+			// Not all platforms have control on the Present Counter, correlate the information FrameCounter -> PresentCounter before Presenting the frame
+			RHIFrameInfoSetPresentCounter(PresentArgs.FrameCounter, GRHIPresentCounter);
+
 			FVulkanViewport* Viewport = ResourceCast(ViewportRHI);
 
 			FVulkanCommandListContext& Context = FVulkanCommandListContext::Get(ExecutingCmdList);
diff --git a/Engine/Source/Runtime/Windows/D3D11RHI/Private/D3D11Viewport.cpp b/Engine/Source/Runtime/Windows/D3D11RHI/Private/D3D11Viewport.cpp
--- a/Engine/Source/Runtime/Windows/D3D11RHI/Private/D3D11Viewport.cpp
+++ b/Engine/Source/Runtime/Windows/D3D11RHI/Private/D3D11Viewport.cpp
@@ -791,6 +791,9 @@ void FD3D11DynamicRHI::RHIEndDrawingViewport(FRHICommandListImmediate& RHICmdLis
 		FD3D11Viewport* Viewport = ResourceCast(ViewportRHI);
 
 		SCOPE_CYCLE_COUNTER(STAT_D3D11PresentTime);
+		
+		// Not all platforms have control on the Present Counter, correlate the information FrameCounter -> PresentCounter before Presenting the frame
+		RHIFrameInfoSetPresentCounter(PresentArgs.FrameCounter, GRHIPresentCounter);
 
 		// Clear references the device might have to resources.
 		CurrentDepthTexture = nullptr;
