Remove extra Present TRACE_BOOKMARK

#Synced-CL 50646776

diff --git a/Engine/Source/Runtime/D3D12RHI/Private/D3D12Viewport.cpp b/Engine/Source/Runtime/D3D12RHI/Private/D3D12Viewport.cpp
--- a/Engine/Source/Runtime/D3D12RHI/Private/D3D12Viewport.cpp
+++ b/Engine/Source/Runtime/D3D12RHI/Private/D3D12Viewport.cpp
@@ -714,8 +714,6 @@ bool FD3D12Viewport::EndRendering(FD3D12CommandContextBase& ContextBase, const F
 		<< End_Rendering.FrameCounter(InPresentArgs.FrameCounter);
 #endif // CPUPROFILERTRACE_ENABLED
 
-	TRACE_BOOKMARK(TEXT("Present P:%" UINT64_FMT " F:%" UINT64_FMT), GRHIPresentCounter, InPresentArgs.FrameCounter)
-
 	check(ContextBase.GetParentAdapter() == GetParentAdapter());
 
 	const bool bCustomPresent = IsValidRef(CustomPresent);
