Add r.D3D12.DXR.MinimumDriverVersionIntel and change existing vendor CVars from int32 to FString

#Synced-CL 50763424

diff --git a/Engine/Source/Runtime/D3D12RHI/Private/Windows/WindowsD3D12Device.cpp b/Engine/Source/Runtime/D3D12RHI/Private/Windows/WindowsD3D12Device.cpp
--- a/Engine/Source/Runtime/D3D12RHI/Private/Windows/WindowsD3D12Device.cpp
+++ b/Engine/Source/Runtime/D3D12RHI/Private/Windows/WindowsD3D12Device.cpp
@@ -64,7 +64,7 @@ static FAutoConsoleVariableRef CVarMinBuildVersionForRayTracing(
 	ECVF_ReadOnly | ECVF_RenderThreadSafe
 );
 
-int32 GMinimumDriverVersionForRayTracingNVIDIA = 0;
+FString GMinimumDriverVersionForRayTracingNVIDIA;
 static FAutoConsoleVariableRef CVarMinDriverVersionForRayTracingNVIDIA(
 	TEXT("r.D3D12.DXR.MinimumDriverVersionNVIDIA"),
 	GMinimumDriverVersionForRayTracingNVIDIA,
@@ -72,9 +72,7 @@ static FAutoConsoleVariableRef CVarMinDriverVersionForRayTracingNVIDIA(
 	ECVF_ReadOnly | ECVF_RenderThreadSafe
 );
 
-// Use AGS_MAKE_VERSION() macro to define the version.
-// i.e. AGS_MAKE_VERSION(major, minor, patch) ((major << 22) | (minor << 12) | patch)
-int32 GMinimumDriverVersionForRayTracingAMD = 0;
+FString GMinimumDriverVersionForRayTracingAMD;
 static FAutoConsoleVariableRef CVarMinDriverVersionForRayTracingAMD(
 	TEXT("r.D3D12.DXR.MinimumDriverVersionAMD"),
 	GMinimumDriverVersionForRayTracingAMD,
@@ -82,6 +80,14 @@ static FAutoConsoleVariableRef CVarMinDriverVersionForRayTracingAMD(
 	ECVF_ReadOnly | ECVF_RenderThreadSafe
 );
 
+FString GMinimumDriverVersionForRayTracingIntel;
+static FAutoConsoleVariableRef CVarMinDriverVersionForRayTracingIntel(
+	TEXT("r.D3D12.DXR.MinimumDriverVersionIntel"),
+	GMinimumDriverVersionForRayTracingIntel,
+	TEXT("Sets the minimum driver version required to enable ray tracing on Intel GPUs."),
+	ECVF_ReadOnly | ECVF_RenderThreadSafe
+);
+
 #if !UE_BUILD_SHIPPING
 static TAutoConsoleVariable<int32> CVarExperimentalShaderModels(
 	TEXT("r.D3D12.ExperimentalShaderModels"),
@@ -907,7 +913,6 @@ void FD3D12DynamicRHIModule::FindAdapter()
 	}
 
 	const EGpuVendorId PreferredVendor = RHIGetPreferredAdapterVendor();
-	const bool bAllowVendorDevice = !FParse::Param(FCommandLine::Get(), TEXT("novendordevice"));
 
 	// Enumerate the DXGIFactory's adapters.
 	for (uint32 AdapterIndex = 0; FD3D12AdapterDesc::EnumAdapters(AdapterIndex, GpuPreference, DXGIFactory4, DXGIFactory6, TempAdapter.GetInitReference()) != DXGI_ERROR_NOT_FOUND; ++AdapterIndex)
@@ -922,7 +927,7 @@ void FD3D12DynamicRHIModule::FindAdapter()
 
 #if INTEL_EXTENSIONS
 			// Enable Intel App Discovery
-			if (AdapterDesc.VendorId == (uint32)EGpuVendorId::Intel && bAllowVendorDevice)
+			if (AdapterDesc.VendorId == (uint32)EGpuVendorId::Intel && UE::RHICore::AllowVendorDevice())
 			{
 				// Intel's App information needs to be registered *before* device creation, so we have to do it here at the last second.
 				// Even though it takes the device ID as an argument, we've been told by Intel that this isn't going to cause problems if multiple Intel devices are detected.
@@ -1532,6 +1537,8 @@ void FD3D12DynamicRHI::Init()
 		UE_LOG(LogD3D12RHI, Warning, TEXT("Ray tracing is disabled because it requires Windows 10 version %u"), (uint32)GMinimumWindowsBuildVersionForRayTracing);
 	}
 
+	const FDriverVersion CurrentDriverVersion(GRHIAdapterUserDriverVersion);
+
 #if WITH_NVAPI
 	if (IsRHIDeviceNVIDIA() && UE::RHICore::AllowVendorDevice())
 	{
@@ -1709,21 +1716,10 @@ void FD3D12DynamicRHI::Init()
 			UE_LOG(LogD3D12RHI, Warning, TEXT("Failed to initialize NVAPI"));
 		}
 
-		NvU32 DriverVersion = UINT32_MAX;
-
-		if (NvStatus == NVAPI_OK)
-		{
-			NvAPI_ShortString BranchString("");
-			if (NvAPI_SYS_GetDriverAndBranchVersion(&DriverVersion, BranchString) != NVAPI_OK)
-			{
-				UE_LOG(LogD3D12RHI, Warning, TEXT("Failed to query NVIDIA driver version"));
-			}
-		}
-
 		// Disable ray tracing for old Nvidia drivers
 		if (GRHISupportsRayTracing 
-			&& GMinimumDriverVersionForRayTracingNVIDIA > 0
-			&& DriverVersion < (uint32)GMinimumDriverVersionForRayTracingNVIDIA)
+			&& !GMinimumDriverVersionForRayTracingNVIDIA.IsEmpty()
+			&& CurrentDriverVersion < FDriverVersion(GMinimumDriverVersionForRayTracingNVIDIA))
 		{
 			DisableRayTracingSupport();
 			UE_LOG(LogD3D12RHI, Warning, TEXT("Ray tracing is disabled because the driver is too old"));
@@ -1731,7 +1727,7 @@ void FD3D12DynamicRHI::Init()
 
 		// Disable indirect ray tracing dispatch on drivers that have a known bug.
 		if (GRHISupportsRayTracingDispatchIndirect 
-			&& DriverVersion < (uint32)46611)
+			&& CurrentDriverVersion < FDriverVersion(TEXT("466.11")))
 		{
 			GRHISupportsRayTracingDispatchIndirect = false;
 			UE_LOG(LogD3D12RHI, Warning,
@@ -1785,7 +1781,7 @@ void FD3D12DynamicRHI::Init()
 		{
 			UE_LOG(LogD3D12RHI, Error, TEXT("Couldn't find CVar r.D3D12.Submission.MaxExecuteBatchSize.Async to work around a known bug in certain drivers."));
 		}
-		else if (FDriverVersion(GRHIGlobals.GpuInfo.AdapterUserDriverVersion) < FDriverVersion(TEXT("571.23")))
+		else if (CurrentDriverVersion < FDriverVersion(TEXT("571.23")))
 		{
 			CVarD3D12SubmissionAsyncMaxExecuteBatchSize->Set(1);
 			UE_LOG(LogD3D12RHI, Warning, TEXT("Batched command list execution is disabled for async queues due to known bugs in the current driver."));
@@ -1795,17 +1791,15 @@ void FD3D12DynamicRHI::Init()
 #if WITH_AMD_AGS
 	if (GRHISupportsRayTracing
 		&& IsRHIDeviceAMD()
-		&& AmdAgsContext
-		&& AmdAgsGpuInfo.radeonSoftwareVersion)
+		&& AmdAgsContext)
 	{
-		if (GMinimumDriverVersionForRayTracingAMD > 0 
-			&& agsCheckDriverVersion(AmdAgsGpuInfo.radeonSoftwareVersion, GMinimumDriverVersionForRayTracingAMD) == AGS_SOFTWAREVERSIONCHECK_OLDER)
+		if (!GMinimumDriverVersionForRayTracingAMD.IsEmpty() &&
+			CurrentDriverVersion < FDriverVersion(GMinimumDriverVersionForRayTracingAMD))
 		{
 			DisableRayTracingSupport();
 			UE_LOG(LogD3D12RHI, Warning, TEXT("Ray tracing is disabled because the driver is too old"));
 		}
-
-		if (GRHISupportsRayTracing)
+		else
 		{
 			static_assert(sizeof(AGSDX12ReturnedParams::ExtensionsSupported) == sizeof(uint32));
 			AGSDX12ReturnedParams::ExtensionsSupported AMDSupportedExtensions;
@@ -1817,19 +1811,28 @@ void FD3D12DynamicRHI::Init()
 	}
 #endif // WITH_AMD_AGS
 
-#if INTEL_EXTENSIONS
-	if (IsRHIDeviceIntel() && UE::RHICore::AllowVendorDevice())
+	if (IsRHIDeviceIntel())
 	{
-		if( GDX12INTCAtomicUInt64Emulation )
+		if (!GMinimumDriverVersionForRayTracingIntel.IsEmpty() &&
+			CurrentDriverVersion < FDriverVersion(GMinimumDriverVersionForRayTracingIntel))
 		{
-			bHasVendorSupportForAtomic64 = GDX12INTCAtomicUInt64Emulation;
+			DisableRayTracingSupport();
+		}
+
+#if INTEL_EXTENSIONS
+		if (UE::RHICore::AllowVendorDevice())
+		{
+			if (GDX12INTCAtomicUInt64Emulation)
+			{
+				bHasVendorSupportForAtomic64 = GDX12INTCAtomicUInt64Emulation;
+			}
+
+			// Create a new Intel Extensions Device Extension Context to support DX12 extension calls
+			INTCExtensionInfo INTCExtensionInfo{};
+			IntelExtensionContext = CreateIntelExtensionsContext(GetAdapter().GetD3DDevice(), INTCExtensionInfo);
 		}
-		
-		// Create a new Intel Extensions Device Extension Context to support DX12 extension calls
-		INTCExtensionInfo INTCExtensionInfo{};
-		IntelExtensionContext = CreateIntelExtensionsContext(GetAdapter().GetD3DDevice(), INTCExtensionInfo);
-	}
 #endif // INTEL_EXTENSIONS
+	}
 
 	GRHIPersistentThreadGroupCount = 1440; // TODO: Revisit based on vendor/adapter/perf query
 
