Add the GPU VRam Used and Budget debug information in stat unit for both DX12 and Vulkan RHIs.

#Synced-CL 51088875

diff --git a/Engine/Source/Runtime/D3D12RHI/Private/D3D12Adapter.cpp b/Engine/Source/Runtime/D3D12RHI/Private/D3D12Adapter.cpp
--- a/Engine/Source/Runtime/D3D12RHI/Private/D3D12Adapter.cpp
+++ b/Engine/Source/Runtime/D3D12RHI/Private/D3D12Adapter.cpp
@@ -1825,7 +1825,7 @@ void FD3D12Adapter::ReleaseTransientUniformBufferAllocator(FTransientUniformBuff
 	verify(TransientUniformBufferAllocators.Remove(InAllocator) == 1);
 }
 
-const FD3DMemoryStats& FD3D12Adapter::CollectMemoryStats()
+const FRHIMemoryStats& FD3D12Adapter::CollectMemoryStats()
 {
 #if PLATFORM_WINDOWS
 	const uint64 UpdateFrame = FrameFence != nullptr ? FrameFence->GetNextFenceToSignal() : 0;
diff --git a/Engine/Source/Runtime/D3D12RHI/Private/D3D12Adapter.h b/Engine/Source/Runtime/D3D12RHI/Private/D3D12Adapter.h
--- a/Engine/Source/Runtime/D3D12RHI/Private/D3D12Adapter.h
+++ b/Engine/Source/Runtime/D3D12RHI/Private/D3D12Adapter.h
@@ -450,8 +450,8 @@ public:
 
 	void BlockUntilIdle();
 
-	const FD3DMemoryStats& CollectMemoryStats();
-	FORCEINLINE const FD3DMemoryStats& GetMemoryStats() const { return MemoryStats; }
+	const FRHIMemoryStats& CollectMemoryStats();
+	FORCEINLINE const FRHIMemoryStats& GetMemoryStats() const { return MemoryStats; }
 	FORCEINLINE uint64 GetMemoryStatsUpdateFrame() const { return MemoryStatsUpdateFrame; }
 
 	bool IsTrackingAllAllocations() const { return bTrackAllAllocation; }
@@ -655,7 +655,7 @@ protected:
 	TArray<FReleasedAllocationData> ReleasedAllocationData;
 	FCriticalSection TrackedAllocationDataCS;
 
-	FD3DMemoryStats MemoryStats;
+	FRHIMemoryStats MemoryStats;
 	uint64 MemoryStatsUpdateFrame;
 
 	TArray<FTransientUniformBufferAllocator*> TransientUniformBufferAllocators;
diff --git a/Engine/Source/Runtime/D3D12RHI/Private/D3D12CommandContext.cpp b/Engine/Source/Runtime/D3D12RHI/Private/D3D12CommandContext.cpp
--- a/Engine/Source/Runtime/D3D12RHI/Private/D3D12CommandContext.cpp
+++ b/Engine/Source/Runtime/D3D12RHI/Private/D3D12CommandContext.cpp
@@ -705,7 +705,7 @@ void FD3D12DynamicRHI::UpdateMemoryStats()
 	if (FD3D12Adapter* const Adapter = ChosenAdapter.Get())
 	{
 		// Refresh captured memory stats.
-		const FD3DMemoryStats& MemoryStats = Adapter->CollectMemoryStats();
+		MemoryStats = Adapter->CollectMemoryStats();
 		UpdateD3DMemoryStatsAndCSV(MemoryStats, true);
 	
 #if STATS
diff --git a/Engine/Source/Runtime/D3D12RHI/Private/D3D12RHI.cpp b/Engine/Source/Runtime/D3D12RHI/Private/D3D12RHI.cpp
--- a/Engine/Source/Runtime/D3D12RHI/Private/D3D12RHI.cpp
+++ b/Engine/Source/Runtime/D3D12RHI/Private/D3D12RHI.cpp
@@ -362,6 +362,11 @@ void FD3D12DynamicRHI::Shutdown()
 #endif
 }
 
+void FD3D12DynamicRHI::RHIGetMemoryStats(FRHIMemoryStats& OutStats)
+{
+	OutStats = MemoryStats;
+}
+
 FD3D12CommandContext* FD3D12DynamicRHI::CreateCommandContext(FD3D12Device* InParent, ED3D12QueueType InQueueType, bool InIsDefaultContext)
 {
 	return new FD3D12CommandContext(InParent, InQueueType, InIsDefaultContext);
diff --git a/Engine/Source/Runtime/D3D12RHI/Private/D3D12RHIPrivate.h b/Engine/Source/Runtime/D3D12RHI/Private/D3D12RHIPrivate.h
--- a/Engine/Source/Runtime/D3D12RHI/Private/D3D12RHIPrivate.h
+++ b/Engine/Source/Runtime/D3D12RHI/Private/D3D12RHIPrivate.h
@@ -234,7 +234,7 @@ private:
 
 	FCriticalSection ObjectsToDeleteCS;
 	TArray<FD3D12DeferredDeleteObject> ObjectsToDelete;
-
+	
 	virtual void UpdateMemoryStats();
 
 #if WITH_NVAPI_RAYTRACING_VALIDATION
@@ -374,6 +374,8 @@ public:
 
 	virtual uint32 RHIComputeMemorySize(FRHITexture* TextureRHI) final override;
 
+	virtual void RHIGetMemoryStats(FRHIMemoryStats& OutStats) final override;
+
 	virtual FRHILockTextureResult RHILockTexture(FRHICommandListImmediate& RHICmdList, const FRHILockTextureArgs& Arguments) final override;
 	virtual void                  RHIUnlockTexture(FRHICommandListImmediate& RHICmdList, const FRHILockTextureArgs& Arguments) final override;
 
@@ -554,6 +556,8 @@ public:
 protected:
 
 	TSharedPtr<FD3D12Adapter> ChosenAdapter;
+	
+	FRHIMemoryStats MemoryStats;
 
 #if D3D12RHI_SUPPORTS_WIN_PIX
 	void* WinPixGpuCapturerHandle = nullptr;
diff --git a/Engine/Source/Runtime/D3D12RHI/Private/D3D12Texture.cpp b/Engine/Source/Runtime/D3D12RHI/Private/D3D12Texture.cpp
--- a/Engine/Source/Runtime/D3D12RHI/Private/D3D12Texture.cpp
+++ b/Engine/Source/Runtime/D3D12RHI/Private/D3D12Texture.cpp
@@ -390,16 +390,16 @@ void FD3D12DynamicRHI::RHIGetTextureMemoryStats(FTextureMemoryStats& OutStats)
 #if PLATFORM_WINDOWS
 	if (GAdjustTexturePoolSizeBasedOnBudget)
 	{
-		const FD3DMemoryStats& MemoryStats = GetAdapter().CollectMemoryStats();
+		const FRHIMemoryStats& CurrentMemoryStats = GetAdapter().CollectMemoryStats();
 
 		// Applications must explicitly manage their usage of physical memory and keep usage within the budget 
 		// assigned to the application process. Processes that cannot keep their usage within their assigned budgets 
 		// will likely experience stuttering, as they are intermittently frozen and paged out to allow other processes to run.
-		const int64 TargetBudget = MemoryStats.BudgetLocal * 0.90f;	// Target using 90% of our budget to account for some fragmentation.
+		const int64 TargetBudget = CurrentMemoryStats.BudgetLocal * 0.90f;	// Target using 90% of our budget to account for some fragmentation.
 		OutStats.TotalGraphicsMemory = TargetBudget;
 
 		const int64 BudgetPadding = TargetBudget * 0.05f;
-		const int64 AvailableSpace = TargetBudget - int64(MemoryStats.UsedLocal);	// Note: AvailableSpace can be negative
+		const int64 AvailableSpace = TargetBudget - int64(CurrentMemoryStats.UsedLocal);	// Note: AvailableSpace can be negative
 		const int64 PreviousTexturePoolSize = RequestedTexturePoolSize;
 		const bool bOverbudget = AvailableSpace < 0;
 
@@ -408,12 +408,12 @@ void FD3D12DynamicRHI::RHIGetTextureMemoryStats(FTextureMemoryStats& OutStats)
 		if (bOverbudget)
 		{
 			// Attempt to lower the texture pool size to meet the budget.
-			const bool bOverActualBudget = MemoryStats.UsedLocal > MemoryStats.BudgetLocal;
+			const bool bOverActualBudget = CurrentMemoryStats.UsedLocal > CurrentMemoryStats.BudgetLocal;
 			UE_CLOG(bOverActualBudget, LogD3D12RHI, Warning,
 				TEXT("Video memory usage is overbudget by %llu MB (using %lld MB/%lld MB budget). Usage breakdown: %lld MB (Streaming Textures), %lld MB (Non Streaming Textures). Last requested texture pool size is %lld MB. This can cause stuttering due to paging."),
-				(MemoryStats.UsedLocal - MemoryStats.BudgetLocal) / 1024ll / 1024ll,
-				MemoryStats.UsedLocal / 1024ll / 1024ll,
-				MemoryStats.BudgetLocal / 1024ll / 1024ll,
+				(CurrentMemoryStats.UsedLocal - CurrentMemoryStats.BudgetLocal) / 1024ll / 1024ll,
+				CurrentMemoryStats.UsedLocal / 1024ll / 1024ll,
+				CurrentMemoryStats.BudgetLocal / 1024ll / 1024ll,
 				GRHIGlobals.StreamingTextureMemorySizeInKB / 1024ll,
 				GRHIGlobals.NonStreamingTextureMemorySizeInKB / 1024ll,
 				PreviousTexturePoolSize / 1024ll / 1024ll);
@@ -423,7 +423,7 @@ void FD3D12DynamicRHI::RHIGetTextureMemoryStats(FTextureMemoryStats& OutStats)
 
 			UE_CLOG(bOverActualBudget && (OutStats.TexturePoolSize >= PreviousTexturePoolSize) && (OutStats.TexturePoolSize > MinTexturePoolSize), LogD3D12RHI, Fatal,
 				TEXT("Video memory usage is overbudget by %llu MB and the texture pool size didn't shrink."),
-				(MemoryStats.UsedLocal - MemoryStats.BudgetLocal) / 1024ll / 1024ll);
+				(CurrentMemoryStats.UsedLocal - CurrentMemoryStats.BudgetLocal) / 1024ll / 1024ll);
 		}
 		else if (AvailableSpace > BudgetPadding)
 		{
diff --git a/Engine/Source/Runtime/D3D12RHI/Private/D3D12Util.cpp b/Engine/Source/Runtime/D3D12RHI/Private/D3D12Util.cpp
--- a/Engine/Source/Runtime/D3D12RHI/Private/D3D12Util.cpp
+++ b/Engine/Source/Runtime/D3D12RHI/Private/D3D12Util.cpp
@@ -738,7 +738,7 @@ void LogPageFaultData(FD3D12Adapter* InAdapter, FD3D12Device* InDevice, D3D12_GP
 
 void LogMemoryStats(FD3D12Adapter* InAdapter)
 {	
-	const FD3DMemoryStats& MemoryStats = InAdapter->GetMemoryStats();
+	const FRHIMemoryStats& MemoryStats = InAdapter->GetMemoryStats();
 
 	UE_LOG(LogD3D12RHI, Error, TEXT("Video Memory Stats from frame ID %d:"), InAdapter->GetMemoryStatsUpdateFrame());
 	UE_LOG(LogD3D12RHI, Error, TEXT("\tLocal Budget:\t%7.2f MB"), MemoryStats.BudgetLocal / (1024.0f * 1024));
@@ -773,8 +773,8 @@ void FD3D12DynamicRHI::TerminateOnOutOfMemory(ID3D12Device* InDevice, HRESULT D3
 
 	if (Adapter)
 	{
-		const auto& MemoryStats = Adapter->GetMemoryStats();
-		FCoreDelegates::GetGPUOutOfMemoryDelegate().Broadcast(MemoryStats.BudgetLocal, MemoryStats.UsedLocal);
+		const auto& CurrentMemoryStats = Adapter->GetMemoryStats();
+		FCoreDelegates::GetGPUOutOfMemoryDelegate().Broadcast(CurrentMemoryStats.BudgetLocal, CurrentMemoryStats.UsedLocal);
 	}
 
 	if (!FApp::IsUnattended())
diff --git a/Engine/Source/Runtime/Engine/Private/UnrealClient.cpp b/Engine/Source/Runtime/Engine/Private/UnrealClient.cpp
--- a/Engine/Source/Runtime/Engine/Private/UnrealClient.cpp
+++ b/Engine/Source/Runtime/Engine/Private/UnrealClient.cpp
@@ -833,6 +833,27 @@ int32 FStatUnitData::DrawStat(FViewport* InViewport, FCanvas* InCanvas, int32 In
 			}
 		}
 
+#if PLATFORM_DESKTOP
+		{
+			FRHIMemoryStats RHIMemStats;
+			RHIGetMemoryStats(RHIMemStats);
+
+			if (RHIMemStats.IsValid())
+			{
+				const bool bOverBudget = RHIMemStats.IsOverBudget();
+				const FColor Color = bOverBudget ? StatRed : StatGreen;
+				DrawTitleString(TEXT("VRam"), NoUnitGraphColor);
+				// Show "Used / Budget" format
+				FString VRAMText = FString::Printf(TEXT("%s / %s"),
+					*GetMemoryString(RHIMemStats.UsedLocal),
+					*GetMemoryString(RHIMemStats.BudgetLocal));
+				DrawDefaultAvgCell(VRAMText, Color);
+
+				PushRows(/* RowsCount = */ 1);
+			}
+		}
+#endif
+
 		// Dynamic resolution
 		if (DynamicResolutionStateInfos.Status == EDynamicResolutionStatus::Enabled || DynamicResolutionStateInfos.Status == EDynamicResolutionStatus::DebugForceEnabled)
 		{
diff --git a/Engine/Source/Runtime/RHI/Private/RHIStats.cpp b/Engine/Source/Runtime/RHI/Private/RHIStats.cpp
--- a/Engine/Source/Runtime/RHI/Private/RHIStats.cpp
+++ b/Engine/Source/Runtime/RHI/Private/RHIStats.cpp
@@ -64,7 +64,7 @@ DEFINE_STAT(STAT_D3DDemotedSystemMemory);
 
 CSV_DEFINE_CATEGORY(GPUMem, true);
 
-void UpdateD3DMemoryStatsAndCSV(const FD3DMemoryStats& MemoryStats, bool bUpdateCSV)
+void UpdateD3DMemoryStatsAndCSV(const FRHIMemoryStats& MemoryStats, bool bUpdateCSV)
 {
 #if STATS || CSV_PROFILER_STATS
 	SCOPE_CYCLE_COUNTER(STAT_D3DUpdateVideoMemoryStats);
diff --git a/Engine/Source/Runtime/RHI/Public/DynamicRHI.h b/Engine/Source/Runtime/RHI/Public/DynamicRHI.h
--- a/Engine/Source/Runtime/RHI/Public/DynamicRHI.h
+++ b/Engine/Source/Runtime/RHI/Public/DynamicRHI.h
@@ -555,6 +555,8 @@ public:
 	// FlushType: Thread safe
 	virtual uint32 RHIComputeMemorySize(FRHITexture* TextureRHI) = 0;
 
+	virtual void RHIGetMemoryStats(FRHIMemoryStats& OutStats) { /*OutStats*/ }
+
 	virtual FRHILockTextureResult RHILockTexture  (FRHICommandListImmediate& RHICmdList, const FRHILockTextureArgs& Arguments) = 0;
 	virtual void                  RHIUnlockTexture(FRHICommandListImmediate& RHICmdList, const FRHILockTextureArgs& Arguments) = 0;
 
@@ -1226,6 +1228,11 @@ inline uint32 RHIComputeMemorySize(FRHITexture* TextureRHI)
 	return GDynamicRHI->RHIComputeMemorySize(TextureRHI);
 }
 
+inline void RHIGetMemoryStats(FRHIMemoryStats& OutStats)
+{
+	GDynamicRHI->RHIGetMemoryStats(OutStats);
+}
+
 inline bool RHIGetRenderQueryResult(FRHIRenderQuery* RenderQuery, uint64& OutResult, bool bWait, uint32 GPUIndex = INDEX_NONE)
 {
 	TRACE_CPUPROFILER_EVENT_SCOPE_CONDITIONAL(RHIGetRenderQueryResult_GPU_Wait, bWait);
diff --git a/Engine/Source/Runtime/RHI/Public/RHIStats.h b/Engine/Source/Runtime/RHI/Public/RHIStats.h
--- a/Engine/Source/Runtime/RHI/Public/RHIStats.h
+++ b/Engine/Source/Runtime/RHI/Public/RHIStats.h
@@ -179,25 +179,21 @@ DECLARE_MEMORY_STAT_POOL_EXTERN(TEXT("Bindless Resource Heap"), STAT_BindlessRes
 DECLARE_DWORD_ACCUMULATOR_STAT_EXTERN(TEXT("Bindless Sampler Descriptors Allocated"), STAT_BindlessSamplerDescriptorsAllocated, STATGROUP_RHI, RHI_API);
 DECLARE_DWORD_ACCUMULATOR_STAT_EXTERN(TEXT("Bindless Resource Descriptors Allocated"), STAT_BindlessResourceDescriptorsAllocated, STATGROUP_RHI, RHI_API);
 
-#if PLATFORM_MICROSOFT
-
-// D3D memory stats.
-struct FD3DMemoryStats
+struct FRHIMemoryStats
 {
-	// Budget assigned by the OS. This can be considered the total memory
-	// the application should use, but an application can also go over-budget.
-	uint64 BudgetLocal = 0;
-	uint64 BudgetSystem = 0;
+	// Budget assigned by the OS/driver. Total memory the app should use.
+	uint64 BudgetLocal = 0;   // GPU VRAM budget
+	uint64 BudgetSystem = 0;  // System/host memory budget
 
-	// Used memory.
-	uint64 UsedLocal = 0;
-	uint64 UsedSystem = 0;
+	// Currently used memory.
+	uint64 UsedLocal = 0;     // GPU VRAM used
+	uint64 UsedSystem = 0;    // System/host memory used
 
-	// Over-budget memory. This is Budget - Used if Used > Budget.
+	// Over-budget memory
 	uint64 DemotedLocal = 0;
 	uint64 DemotedSystem = 0;
 
-	// Available memory within budget. This is Budget - Used clamped to 0 if over-budget.
+	// Available memory within budget
 	uint64 AvailableLocal = 0;
 	uint64 AvailableSystem = 0;
 
@@ -205,8 +201,15 @@ struct FD3DMemoryStats
 	{
 		return DemotedLocal > 0 || DemotedSystem > 0;
 	}
+
+	bool IsValid() const
+	{
+		return BudgetLocal > 0 || BudgetSystem > 0;
+	}
 };
 
+#if PLATFORM_MICROSOFT
+
 DECLARE_STATS_GROUP(TEXT("D3D Video Memory"), STATGROUP_D3DMemory, STATCAT_Advanced);
 DECLARE_MEMORY_STAT_EXTERN(TEXT("Total Video Memory (Budget)"), STAT_D3DTotalVideoMemory, STATGROUP_D3DMemory, );
 DECLARE_MEMORY_STAT_EXTERN(TEXT("Total System Memory (Budget)"), STAT_D3DTotalSystemMemory, STATGROUP_D3DMemory, );
@@ -219,6 +222,6 @@ DECLARE_MEMORY_STAT_EXTERN(TEXT("Demoted System Memory"), STAT_D3DDemotedSystemM
 DECLARE_CYCLE_STAT_EXTERN(TEXT("Video Memory stats update time"), STAT_D3DUpdateVideoMemoryStats, STATGROUP_D3DMemory, RHI_API);
 
 // Update D3D memory stat counters and CSV profiler stats, if enabled.
-RHI_API void UpdateD3DMemoryStatsAndCSV(const FD3DMemoryStats& MemoryStats, bool bUpdateCSV);
+RHI_API void UpdateD3DMemoryStatsAndCSV(const FRHIMemoryStats& MemoryStats, bool bUpdateCSV);
 
 #endif // PLATFORM_MICROSOFT
\ No newline at end of file
diff --git a/Engine/Source/Runtime/RHICore/Internal/DXGIUtilities.h b/Engine/Source/Runtime/RHICore/Internal/DXGIUtilities.h
--- a/Engine/Source/Runtime/RHICore/Internal/DXGIUtilities.h
+++ b/Engine/Source/Runtime/RHICore/Internal/DXGIUtilities.h
@@ -20,12 +20,12 @@ THIRD_PARTY_INCLUDES_START
 THIRD_PARTY_INCLUDES_END
 #include "Microsoft/HideMicrosoftPlatformTypes.h"
 
-struct FD3DMemoryStats;
+struct FRHIMemoryStats;
 
 namespace UE::DXGIUtilities
 {
 	RHICORE_API const TCHAR* GetFormatString(DXGI_FORMAT Format);
-	RHICORE_API HRESULT GetD3DMemoryStats(IDXGIAdapter* Adapter, FD3DMemoryStats& OutStats);
+	RHICORE_API HRESULT GetD3DMemoryStats(IDXGIAdapter* Adapter, FRHIMemoryStats& OutStats);
 
 	inline DXGI_FORMAT FindSharedResourceFormat(DXGI_FORMAT InFormat, bool bSRGB)
 	{
diff --git a/Engine/Source/Runtime/RHICore/Private/DXGIUtilities.cpp b/Engine/Source/Runtime/RHICore/Private/DXGIUtilities.cpp
--- a/Engine/Source/Runtime/RHICore/Private/DXGIUtilities.cpp
+++ b/Engine/Source/Runtime/RHICore/Private/DXGIUtilities.cpp
@@ -52,7 +52,7 @@ const TCHAR* UE::DXGIUtilities::GetFormatString(DXGI_FORMAT Format)
 	return Result;
 }
 
-HRESULT UE::DXGIUtilities::GetD3DMemoryStats(IDXGIAdapter* Adapter, FD3DMemoryStats& OutStats)
+HRESULT UE::DXGIUtilities::GetD3DMemoryStats(IDXGIAdapter* Adapter, FRHIMemoryStats& OutStats)
 {
 #if PLATFORM_WINDOWS
 	SCOPE_CYCLE_COUNTER(STAT_D3DUpdateVideoMemoryStats);
diff --git a/Engine/Source/Runtime/VulkanRHI/Private/VulkanDynamicRHI.h b/Engine/Source/Runtime/VulkanRHI/Private/VulkanDynamicRHI.h
--- a/Engine/Source/Runtime/VulkanRHI/Private/VulkanDynamicRHI.h
+++ b/Engine/Source/Runtime/VulkanRHI/Private/VulkanDynamicRHI.h
@@ -182,6 +182,8 @@ public:
 	virtual FTextureRHIRef RHIAsyncCreateTexture2D(uint32 SizeX, uint32 SizeY, uint8 Format, uint32 NumMips, ETextureCreateFlags Flags, ERHIAccess InResourceState, void** InitialMipData, uint32 NumInitialMips, const TCHAR* DebugName, FGraphEventRef& OutCompletionEvent) final override;
 	virtual uint32 RHIComputeMemorySize(FRHITexture* TextureRHI) final override;
 
+	virtual void RHIGetMemoryStats(FRHIMemoryStats& OutStats) final override;
+
 	virtual FRHILockTextureResult RHILockTexture(FRHICommandListImmediate& RHICmdList, const FRHILockTextureArgs& Arguments) final override;
 	virtual void RHIUnlockTexture(FRHICommandListImmediate& RHICmdList, const FRHILockTextureArgs& Arguments) final override;
 
diff --git a/Engine/Source/Runtime/VulkanRHI/Private/VulkanMemory.cpp b/Engine/Source/Runtime/VulkanRHI/Private/VulkanMemory.cpp
--- a/Engine/Source/Runtime/VulkanRHI/Private/VulkanMemory.cpp
+++ b/Engine/Source/Runtime/VulkanRHI/Private/VulkanMemory.cpp
@@ -557,7 +557,9 @@ namespace VulkanRHI
 
 	void FDeviceMemoryManager::UpdateMemoryProperties()
 	{
-		if (Device->GetOptionalExtensions().HasMemoryBudget)
+		const bool bHasMemoryBudget = Device->GetOptionalExtensions().HasMemoryBudget;
+
+		if (bHasMemoryBudget)
 		{
 			VkPhysicalDeviceMemoryProperties2 MemoryProperties2;
 			ZeroVulkanStruct(MemoryBudget, VK_STRUCTURE_TYPE_PHYSICAL_DEVICE_MEMORY_BUDGET_PROPERTIES_EXT);
@@ -595,6 +597,51 @@ namespace VulkanRHI
 		{
 			VulkanRHI::vkGetPhysicalDeviceMemoryProperties(Device->GetPhysicalHandle(), &MemoryProperties);
 		}
+
+		// Save Memory stats in FRHIMemoryStats format
+		MemoryStats = FRHIMemoryStats{};
+		uint64 HeapBudget, HeapUsed;
+
+		for (uint32 HeapIndex = 0; HeapIndex < MemoryProperties.memoryHeapCount; ++HeapIndex)
+		{
+			const bool bLocalHeap = (MemoryProperties.memoryHeaps[HeapIndex].flags & VK_MEMORY_HEAP_DEVICE_LOCAL_BIT) != 0;
+
+			if (bHasMemoryBudget)
+			{
+				HeapBudget = MemoryBudget.heapBudget[HeapIndex];
+				HeapUsed = MemoryBudget.heapUsage[HeapIndex];
+			}
+			else
+			{
+				HeapBudget = MemoryProperties.memoryHeaps[HeapIndex].size;
+				
+				if (HeapIndex < (uint32)HeapInfos.Num())
+				{
+					HeapUsed = HeapInfos[HeapIndex].UsedSize;
+				}
+				else
+				{
+					HeapUsed = 0;
+				}
+			}
+
+			if (bLocalHeap)
+			{
+				MemoryStats.BudgetLocal += HeapBudget;
+				MemoryStats.UsedLocal += HeapUsed;
+			}
+			else
+			{
+				MemoryStats.BudgetSystem += HeapBudget;
+				MemoryStats.UsedSystem += HeapUsed;
+			}
+		}
+
+		MemoryStats.AvailableLocal = (MemoryStats.UsedLocal < MemoryStats.BudgetLocal) ? (MemoryStats.BudgetLocal - MemoryStats.UsedLocal) : 0;
+		MemoryStats.AvailableSystem = (MemoryStats.UsedSystem < MemoryStats.BudgetSystem) ? (MemoryStats.BudgetSystem - MemoryStats.UsedSystem) : 0;
+
+		MemoryStats.DemotedLocal = (MemoryStats.UsedLocal > MemoryStats.BudgetLocal) ? (MemoryStats.UsedLocal - MemoryStats.BudgetLocal) : 0;
+		MemoryStats.DemotedSystem = (MemoryStats.UsedSystem > MemoryStats.BudgetSystem) ? (MemoryStats.UsedSystem - MemoryStats.BudgetSystem) : 0;
 	}
 
 	void FDeviceMemoryManager::Init(FVulkanDevice* InDevice)
@@ -774,11 +821,6 @@ namespace VulkanRHI
 		else
 		{
 			OutAllocated = HeapInfos[PrimaryHeapIndex].UsedSize;
-			if (Device->GetOptionalExtensions().HasMemoryBudget && (FPlatformTime::Seconds() - MemoryUpdateTime >= 1.0))
-			{
-				MemoryUpdateTime = FPlatformTime::Seconds();
-				UpdateMemoryProperties();
-			}			
 			OutLimit = GetBaseHeapSize(PrimaryHeapIndex);
 	}
 	}
diff --git a/Engine/Source/Runtime/VulkanRHI/Private/VulkanMemory.h b/Engine/Source/Runtime/VulkanRHI/Private/VulkanMemory.h
--- a/Engine/Source/Runtime/VulkanRHI/Private/VulkanMemory.h
+++ b/Engine/Source/Runtime/VulkanRHI/Private/VulkanMemory.h
@@ -282,7 +282,7 @@ namespace VulkanRHI
 	class FDeviceMemoryManager
 	{
 		TMap<FDeviceMemoryBlockKey, FDeviceMemoryBlock> Allocations;
-		void UpdateMemoryProperties();
+		FRHIMemoryStats MemoryStats;
 	public:
 		FDeviceMemoryManager();
 		~FDeviceMemoryManager();
@@ -327,6 +327,9 @@ namespace VulkanRHI
 		VkDeviceSize GetBaseHeapSize(uint32 HeapIndex) const;
 		uint32 GetHeapIndex(uint32 MemoryTypeIndex);
 
+		void UpdateMemoryProperties();
+		FRHIMemoryStats GetRHIMemoryStats() const { return MemoryStats; }
+
 	protected:
 		void FreeInternal(FDeviceMemoryAllocation* Allocation);
 		void TrimMemory(bool bFullTrim);
diff --git a/Engine/Source/Runtime/VulkanRHI/Private/VulkanRHI.cpp b/Engine/Source/Runtime/VulkanRHI/Private/VulkanRHI.cpp
--- a/Engine/Source/Runtime/VulkanRHI/Private/VulkanRHI.cpp
+++ b/Engine/Source/Runtime/VulkanRHI/Private/VulkanRHI.cpp
@@ -1197,6 +1197,11 @@ void FVulkanDynamicRHI::RHIEndFrame_RenderThread(FRHICommandListImmediate& RHICm
 
 void FVulkanDynamicRHI::RHIEndFrame(const FRHIEndFrameArgs& Args)
 {
+	if (GetDevice())
+	{
+		GetDevice()->DeviceMemoryManager.UpdateMemoryProperties();
+	}
+
 	// Close the previous frame's timing and start a new one
 	auto Lambda = [this, OldTiming = MoveTemp(CurrentTimingPerQueue)]()
 	{
@@ -1654,6 +1659,14 @@ void* FVulkanDynamicRHI::RHIGetNativeComputeQueue()
 		(void*)Device->GetGraphicsQueue()->GetHandle();
 }
 
+void FVulkanDynamicRHI::RHIGetMemoryStats(FRHIMemoryStats& OutStats)
+{
+	if (GetDevice())
+	{
+		OutStats = GetDevice()->GetDeviceMemoryManager().GetRHIMemoryStats();
+	}
+}
+
 IRHICommandContext* FVulkanDynamicRHI::RHIGetDefaultContext()
 {
 	return &Device->GetImmediateContext();
diff --git a/Engine/Source/Runtime/Windows/D3D11RHI/Private/D3D11RHI.cpp b/Engine/Source/Runtime/Windows/D3D11RHI/Private/D3D11RHI.cpp
--- a/Engine/Source/Runtime/Windows/D3D11RHI/Private/D3D11RHI.cpp
+++ b/Engine/Source/Runtime/Windows/D3D11RHI/Private/D3D11RHI.cpp
@@ -360,7 +360,6 @@ void FD3D11DynamicRHI::UpdateMemoryStats()
 {
 #if PLATFORM_WINDOWS && (STATS || CSV_PROFILER_STATS)
 	// Some older drivers don't support querying memory stats, so don't do anything if this fails.
-	FD3DMemoryStats MemoryStats;
 	if (SUCCEEDED(UE::DXGIUtilities::GetD3DMemoryStats(GetAdapter().DXGIAdapter, MemoryStats)))
 	{
 		UpdateD3DMemoryStatsAndCSV(MemoryStats, true);
@@ -368,6 +367,11 @@ void FD3D11DynamicRHI::UpdateMemoryStats()
 #endif // PLATFORM_WINDOWS && (STATS || CSV_PROFILER_STATS)
 }
 
+void FD3D11DynamicRHI::RHIGetMemoryStats(FRHIMemoryStats& OutStats)
+{
+	OutStats = MemoryStats;
+}
+
 #if ENABLE_LOW_LEVEL_MEM_TRACKER || UE_MEMORY_TRACE_ENABLED
 void FD3D11DynamicRHI::RHIUpdateAllocationTags(FRHICommandListBase& RHICmdList, FRHIBuffer* BufferRHI)
 {
diff --git a/Engine/Source/Runtime/Windows/D3D11RHI/Private/D3D11RHIPrivate.h b/Engine/Source/Runtime/Windows/D3D11RHI/Private/D3D11RHIPrivate.h
--- a/Engine/Source/Runtime/Windows/D3D11RHI/Private/D3D11RHIPrivate.h
+++ b/Engine/Source/Runtime/Windows/D3D11RHI/Private/D3D11RHIPrivate.h
@@ -270,6 +270,7 @@ public:
 	[[nodiscard]] D3D11RHI_API virtual FRHITextureInitializer RHICreateTextureInitializer(FRHICommandListBase& RHICmdList, const FRHITextureCreateDesc& CreateDesc) final override;
 	D3D11RHI_API virtual FTextureRHIRef RHIAsyncCreateTexture2D(uint32 SizeX, uint32 SizeY, uint8 Format, uint32 NumMips, ETextureCreateFlags Flags, ERHIAccess InResourceState, void** InitialMipData, uint32 NumInitialMips, const TCHAR* DebugName, FGraphEventRef& OutCompletionEvent) final override;
 	D3D11RHI_API virtual uint32 RHIComputeMemorySize(FRHITexture* TextureRHI) final override;
+	D3D11RHI_API virtual void RHIGetMemoryStats(FRHIMemoryStats& OutStats) final override;
 
 	D3D11RHI_API virtual FRHILockTextureResult RHILockTexture(FRHICommandListImmediate& RHICmdList, const FRHILockTextureArgs& Arguments) final override;
 	D3D11RHI_API virtual void                  RHIUnlockTexture(FRHICommandListImmediate& RHICmdList, const FRHILockTextureArgs& Arguments) final override;
@@ -686,6 +687,7 @@ public:
 
 protected:
 	FD3D11Adapter Adapter;
+	FRHIMemoryStats MemoryStats;
 
 	D3D11RHI_API FD3D11Texture* CreateTextureFromResource(bool bTextureArray, bool bCubeTexture, EPixelFormat Format, ETextureCreateFlags TexCreateFlags, const FClearValueBinding& ClearValueBinding, ID3D11Texture2D* TextureResource);
 
diff --git a/Engine/Source/Runtime/Windows/D3D11RHI/Private/D3D11Util.cpp b/Engine/Source/Runtime/Windows/D3D11RHI/Private/D3D11Util.cpp
--- a/Engine/Source/Runtime/Windows/D3D11RHI/Private/D3D11Util.cpp
+++ b/Engine/Source/Runtime/Windows/D3D11RHI/Private/D3D11Util.cpp
@@ -180,7 +180,7 @@ static void TerminateOnDeviceRemoved(HRESULT D3DResult, ID3D11Device* Direct3DDe
 
 void GetAndLogMemoryStats(const FD3D11Adapter& InAdapter, uint64& OutVRAMBudgetBytes, uint64& OutVRAMUsageBytes)
 {
-	FD3DMemoryStats MemoryStats;
+	FRHIMemoryStats MemoryStats;
 	if (SUCCEEDED(UE::DXGIUtilities::GetD3DMemoryStats(InAdapter.DXGIAdapter, MemoryStats)))
 	{
 		UE_LOG(LogD3D11RHI, Error, TEXT("\tLocal Budget:\t%7.2f MB"), MemoryStats.BudgetLocal / (1024.0f * 1024));
