Ray Tracing: Refactor D3D12 RHI to avoid storage of geometry descriptions

The array of geometry descriptors must be patched before each acceleration structure build anyway because vertex/index buffer can change. Get rid of the copy to streamline the code and make it less error prone.

Remove unused NullTransformBuffer

#rnx
#Synced-CL 51098442

diff --git a/Engine/Source/Runtime/D3D12RHI/Private/D3D12RayTracing.cpp b/Engine/Source/Runtime/D3D12RHI/Private/D3D12RayTracing.cpp
--- a/Engine/Source/Runtime/D3D12RHI/Private/D3D12RayTracing.cpp
+++ b/Engine/Source/Runtime/D3D12RHI/Private/D3D12RayTracing.cpp
@@ -122,6 +122,7 @@ DECLARE_DWORD_ACCUMULATOR_STAT(TEXT("Compiled shaders (total)"), STAT_D3D12RayTr
 DECLARE_DWORD_ACCUMULATOR_STAT(TEXT("Allocated bottom level acceleration structures"), STAT_D3D12RayTracingAllocatedBLAS, STATGROUP_D3D12RayTracing);
 DECLARE_DWORD_ACCUMULATOR_STAT(TEXT("Allocated top level acceleration structures"), STAT_D3D12RayTracingAllocatedTLAS, STATGROUP_D3D12RayTracing);
 DECLARE_DWORD_ACCUMULATOR_STAT(TEXT("Triangles in all BL acceleration structures"), STAT_D3D12RayTracingTrianglesBLAS, STATGROUP_D3D12RayTracing);
+DECLARE_DWORD_ACCUMULATOR_STAT(TEXT("Procedural primitives in all BL acceleration structures"), STAT_D3D12RayTracingProceduralPrimitivesBLAS, STATGROUP_D3D12RayTracing);
 
 DECLARE_DWORD_COUNTER_STAT(TEXT("Built BL AS (per frame)"), STAT_D3D12RayTracingBuiltBLAS, STATGROUP_D3D12RayTracing);
 DECLARE_DWORD_COUNTER_STAT(TEXT("Updated BL AS (per frame)"), STAT_D3D12RayTracingUpdatedBLAS, STATGROUP_D3D12RayTracing);
@@ -450,6 +451,44 @@ static D3D12_RAYTRACING_GEOMETRY_TYPE TranslateRayTracingGeometryType(ERayTracin
 	}
 }
 
+static DXGI_FORMAT TranslateRayTracingVertexElementType(EVertexElementType Type)
+{
+	switch (Type)
+	{
+		case VET_Float4:
+			// While the DXGI_FORMAT_R32G32B32A32_FLOAT format is not supported by DXR, since we manually load vertex 
+			// data when we are building the BLAS, we can just rely on the vertex stride to offset the read index, 
+			// and read only the 3 vertex components, and so use the DXGI_FORMAT_R32G32B32_FLOAT vertex format
+		case VET_Float3:
+			return DXGI_FORMAT_R32G32B32_FLOAT;
+		case VET_Float2:
+			return DXGI_FORMAT_R32G32_FLOAT;
+		case VET_Float1:
+			return DXGI_FORMAT_R32_FLOAT;
+		case VET_Half2:
+			return DXGI_FORMAT_R16G16_FLOAT;
+	}
+	return DXGI_FORMAT_UNKNOWN;
+}
+
+static D3D12_RAYTRACING_GEOMETRY_FLAGS TranslateSegmentToGeometryFlags(const FRayTracingGeometrySegment& Segment)
+{
+	D3D12_RAYTRACING_GEOMETRY_FLAGS Flags = D3D12_RAYTRACING_GEOMETRY_FLAG_NONE;
+
+	if (Segment.bForceOpaque)
+	{
+		// Deny anyhit shader invocations when this segment is hit
+		Flags |= D3D12_RAYTRACING_GEOMETRY_FLAG_OPAQUE;
+	}
+
+	if (!Segment.bAllowDuplicateAnyHitShaderInvocation)
+	{
+		// Allow only a single any-hit shader invocation per primitive
+		Flags |= D3D12_RAYTRACING_GEOMETRY_FLAG_NO_DUPLICATE_ANYHIT_INVOCATION;
+	}
+	return Flags;
+}
+
 struct FDXILLibrary
 {
 	// No copy assignment or move because FDXILLibrary points to internal struct memory
@@ -3282,66 +3321,28 @@ void TranslateRayTracingGeometryDescs(const FRayTracingGeometryInitializer& Init
 			TEXT("FRayTracingGeometrySegment.MaxVertices for '%s' must contain number of positions in the vertex buffer or maximum index buffer value+1 if index buffer is provided."),
 			*Initializer.DebugName.ToString());
 
-		if (Initializer.GeometryType == RTGT_Triangles)
-		{
-			checkf(Segment.VertexBufferElementType == VET_Float3
-				|| Segment.VertexBufferElementType == VET_Float4, TEXT("Only float3/4 vertex buffers are currently implemented.")); // #dxr_todo UE-72160: support other vertex buffer formats
-			checkf(Segment.VertexBufferStride >= 12, TEXT("Only deinterleaved float3 position vertex buffers are currently implemented.")); // #dxr_todo UE-72160: support interleaved vertex buffers
-		}
-		else if (Initializer.GeometryType == RTGT_Procedural)
-		{
-			checkf(Segment.VertexBufferStride >= (2 * sizeof(FVector3f)), TEXT("Procedural geometry vertex buffer must contain at least 2xFloat3 that defines 3D bounding boxes of primitives."));
-		}
+		D3D12_RAYTRACING_GEOMETRY_DESC& Desc = Output[SegmentIndex];
+		Desc = {}; // zero-initialize
 
-		if (Initializer.IndexBuffer)
-		{
-			uint32 IndexStride = Initializer.IndexBuffer->GetStride();
-			check(Initializer.IndexBuffer->GetSize() >=
-				(Segment.FirstPrimitive + Segment.NumPrimitives) * FD3D12RayTracingGeometry::IndicesPerPrimitive * IndexStride + Initializer.IndexBufferOffset);
-		}
+		Desc.Flags = TranslateSegmentToGeometryFlags(Segment);
 
-		D3D12_RAYTRACING_GEOMETRY_DESC Desc = {};
-
-		Desc.Flags = D3D12_RAYTRACING_GEOMETRY_FLAG_NONE;
-		Desc.Type = GeometryType;
-
-		if (Segment.bForceOpaque)
+		switch (Initializer.GeometryType)
 		{
-			// Deny anyhit shader invocations when this segment is hit
-			Desc.Flags |= D3D12_RAYTRACING_GEOMETRY_FLAG_OPAQUE;
-		}
+		case RTGT_Triangles:
+			checkf(Segment.VertexBufferElementType == VET_Float3 ||
+				   Segment.VertexBufferElementType == VET_Float4, TEXT("Only float3/4 vertex buffers are currently implemented.")); // #dxr_todo UE-72160: support other vertex buffer formats
+			checkf(Segment.VertexBufferStride >= 12, TEXT("Only deinterleaved float3 position vertex buffers are currently implemented.")); // #dxr_todo UE-72160: support interleaved vertex buffers
 
-		if (!Segment.bAllowDuplicateAnyHitShaderInvocation)
-		{
-			// Allow only a single any-hit shader invocation per primitive
-			Desc.Flags |= D3D12_RAYTRACING_GEOMETRY_FLAG_NO_DUPLICATE_ANYHIT_INVOCATION;
-		}
-
-		switch (GeometryType)
-		{
-		case D3D12_RAYTRACING_GEOMETRY_TYPE_TRIANGLES:
-			switch (Segment.VertexBufferElementType)
-			{
-			case VET_Float4:
-				// While the DXGI_FORMAT_R32G32B32A32_FLOAT format is not supported by DXR, since we manually load vertex 
-				// data when we are building the BLAS, we can just rely on the vertex stride to offset the read index, 
-				// and read only the 3 vertex components, and so use the DXGI_FORMAT_R32G32B32_FLOAT vertex format
-			case VET_Float3:
-				Desc.Triangles.VertexFormat = DXGI_FORMAT_R32G32B32_FLOAT;
-				break;
-			case VET_Float2:
-				Desc.Triangles.VertexFormat = DXGI_FORMAT_R32G32_FLOAT;
-				break;
-			case VET_Half2:
-				Desc.Triangles.VertexFormat = DXGI_FORMAT_R16G16_FLOAT;
-				break;
-			default:
-				checkNoEntry();
-				break;
-			}
+			Desc.Type = D3D12_RAYTRACING_GEOMETRY_TYPE_TRIANGLES;
+			Desc.Triangles.VertexFormat = TranslateRayTracingVertexElementType(Segment.VertexBufferElementType);
+			checkf(Desc.Triangles.VertexFormat != DXGI_FORMAT_UNKNOWN, TEXT("Unknown vertex format %d - not currently implemented"), int(Desc.Triangles.VertexFormat));
 
 			if (Initializer.IndexBuffer)
 			{
+				uint32 IndexStride = Initializer.IndexBuffer->GetStride();
+				check(Initializer.IndexBuffer->GetSize() >=
+					  (Segment.FirstPrimitive + Segment.NumPrimitives) * FD3D12RayTracingGeometry::IndicesPerPrimitive * IndexStride + Initializer.IndexBufferOffset);
+
 				// In some cases the geometry is created with 16 bit index buffer, but it's 32 bit at build time.
 				// We conservatively set this to 32 bit to allocate acceleration structure memory.
 				Desc.Triangles.IndexFormat = DXGI_FORMAT_R32_UINT;
@@ -3357,24 +3358,24 @@ void TranslateRayTracingGeometryDescs(const FRayTracingGeometryInitializer& Init
 			}
 			break;
 
-		case D3D12_RAYTRACING_GEOMETRY_TYPE_PROCEDURAL_PRIMITIVE_AABBS:
+		case RTGT_Procedural:
+			checkf(Segment.VertexBufferStride >= (2 * sizeof(FVector3f)), TEXT("Procedural geometry vertex buffer must contain at least 2xFloat3 that defines 3D bounding boxes of primitives."));
+
+			Desc.Type = D3D12_RAYTRACING_GEOMETRY_TYPE_PROCEDURAL_PRIMITIVE_AABBS;
 			Desc.AABBs.AABBCount = Segment.NumPrimitives;
 			break;
-
 		default:
 			checkf(false, TEXT("Unexpected ray tracing geometry type"));
 			break;
 		}
 
 		ComputedPrimitiveCountForValidation += Segment.NumPrimitives;
-
-
-		Output[SegmentIndex] = Desc;
 	}
 
 	check(ComputedPrimitiveCountForValidation == Initializer.TotalPrimitiveCount);
 }
 
+
 FRayTracingAccelerationStructureSize FD3D12DynamicRHI::RHICalcRayTracingGeometrySize(const FRayTracingGeometryInitializer& Initializer)
 {
 	{
@@ -3496,8 +3497,6 @@ FShaderBindingTableRHIRef FD3D12DynamicRHI::RHICreateShaderBindingTable(FRHIComm
 	return new FD3D12RayTracingShaderBindingTable(RHICmdList, &Adapter, Initializer);
 }
 
-FBufferRHIRef FD3D12RayTracingGeometry::NullTransformBuffer;
-
 enum class ERayTracingBufferType
 {
 	AccelerationStructure,
@@ -3604,26 +3603,10 @@ FD3D12RayTracingGeometry::FD3D12RayTracingGeometry(FRHICommandListBase& RHICmdLi
 	FMemory::Memzero(bHasPendingCompactionRequests);
 	FMemory::Memzero(bRegisteredAsRenameListener);
 
-	if(!FD3D12RayTracingGeometry::NullTransformBuffer.IsValid())
-	{
-		TArray<float> NullTransformData;
-		NullTransformData.SetNumZeroed(12);
-
-		FD3D12RayTracingGeometry::NullTransformBuffer = UE::RHIResourceUtils::CreateVertexBufferFromArray(
-			RHICmdList,
-			TEXT("NullTransformBuffer"),
-			EBufferUsageFlags::Static,
-			MakeConstArrayView(NullTransformData)
-		);
-	}
-
 	RegisterD3D12RayTracingGeometry(this);
 
 	checkf(Initializer.Segments.Num() > 0, TEXT("Ray tracing geometry must be initialized with at least one segment."));
 
-	GeometryDescs.SetNumUninitialized(Initializer.Segments.Num());
-	TranslateRayTracingGeometryDescs(Initializer, GeometryDescs);
-
 	SetDirty(FRHIGPUMask::All(), true);
 
 	const bool bHasOfflineMetadata = Initializer.OfflineDataHeader.IsValid();
@@ -3696,7 +3679,11 @@ FD3D12RayTracingGeometry::FD3D12RayTracingGeometry(FRHICommandListBase& RHICmdLi
 		}
 	});
 
-	INC_DWORD_STAT_BY(STAT_D3D12RayTracingTrianglesBLAS, Initializer.TotalPrimitiveCount);
+	switch (Initializer.GeometryType)
+	{
+		case RTGT_Triangles: 	INC_DWORD_STAT_BY(STAT_D3D12RayTracingTrianglesBLAS				, Initializer.TotalPrimitiveCount); break;
+		case RTGT_Procedural:	INC_DWORD_STAT_BY(STAT_D3D12RayTracingProceduralPrimitivesBLAS	, Initializer.TotalPrimitiveCount); break;
+	}
 
 	const bool bForRendering = Initializer.Type == ERayTracingGeometryInitializerType::Rendering;
 	if (SourceData != nullptr)
@@ -3766,9 +3753,6 @@ void FD3D12RayTracingGeometry::Swap(FD3D12RayTracingGeometry& Other)
 
 	checkf(Initializer.Segments.Num() > 0, TEXT("Ray tracing geometry must be initialized with at least one segment."));
 
-	GeometryDescs.SetNumUninitialized(Initializer.Segments.Num());
-	TranslateRayTracingGeometryDescs(Initializer, GeometryDescs);
-		
 	FOREACH_GPU(GPUIndex < MAX_NUM_GPUS && GPUIndex < GNumExplicitGPUsForRendering,
 	{
 		RegisterAsRenameListener(GPUIndex);
@@ -3823,7 +3807,12 @@ void FD3D12RayTracingGeometry::ReleaseUnderlyingResource()
 
 	if (Initializer.Type != ERayTracingGeometryInitializerType::StreamingSource)
 	{
-		DEC_DWORD_STAT_BY(STAT_D3D12RayTracingTrianglesBLAS, Initializer.TotalPrimitiveCount);
+		switch (Initializer.GeometryType)
+		{
+			case RTGT_Triangles: 	DEC_DWORD_STAT_BY(STAT_D3D12RayTracingTrianglesBLAS				, Initializer.TotalPrimitiveCount); break;
+			case RTGT_Procedural:	DEC_DWORD_STAT_BY(STAT_D3D12RayTracingProceduralPrimitivesBLAS	, Initializer.TotalPrimitiveCount); break;
+		}
+
 		DEC_DWORD_STAT(STAT_D3D12RayTracingAllocatedBLAS);
 	}
 
@@ -3836,7 +3825,6 @@ void FD3D12RayTracingGeometry::ReleaseUnderlyingResource()
 	Initializer = {};
 
 	AccelerationStructureCompactedSize = 0;
-	GeometryDescs = {};
 	for (TArray<FD3D12HitGroupSystemParameters>& HitGroupParametersForGPU : HitGroupSystemParameters)
 	{
 		HitGroupParametersForGPU.Empty();
@@ -3881,12 +3869,14 @@ bool FD3D12RayTracingGeometry::AllocateBufferSRVs(uint32 InGPUIndex)
 			MaxPrimitiveCount = FMath::Max(MaxPrimitiveCount, Segment.FirstPrimitive + Segment.NumPrimitives);
 		}
 
+		const uint32 PerPrimitiveIndexCount = 3;
+
 		D3D12_SHADER_RESOURCE_VIEW_DESC SRVDesc = {};
 		SRVDesc.ViewDimension = D3D12_SRV_DIMENSION_BUFFER;
 		SRVDesc.Format = DXGI_FORMAT_R32_TYPELESS;
 		SRVDesc.Shader4ComponentMapping = D3D12_DEFAULT_SHADER_4_COMPONENT_MAPPING;
 		SRVDesc.Buffer.FirstElement = (Initializer.IndexBufferOffset + IndexBuffer->ResourceLocation.GetOffsetFromBaseOfResource()) >> 2u;
-		SRVDesc.Buffer.NumElements = FMath::Max((uint32)1, ((MaxPrimitiveCount * 3 * IndexBuffer->GetStride()) + 3) >> 2u);
+		SRVDesc.Buffer.NumElements = FMath::Max((uint32)1, ((MaxPrimitiveCount * PerPrimitiveIndexCount * IndexBuffer->GetStride()) + 3) >> 2u);
 		SRVDesc.Buffer.Flags = D3D12_BUFFER_SRV_FLAG_RAW;
 		SRVDesc.Buffer.StructureByteStride = 0;
 
@@ -4092,8 +4082,6 @@ void FD3D12RayTracingGeometry::UpdateResidency(FD3D12CommandContext& CommandCont
 
 void FD3D12RayTracingGeometry::SetupHitGroupSystemParameters(uint32 InGPUIndex)
 {
-	D3D12_RAYTRACING_GEOMETRY_TYPE GeometryType = TranslateRayTracingGeometryType(Initializer.GeometryType);
-
 	bool bBindless = AreBindlessResourcesEnabledForRayTracing(GetParentAdapter());
 
 	bool bAnyBuffersUpdated = true;
@@ -4110,6 +4098,7 @@ void FD3D12RayTracingGeometry::SetupHitGroupSystemParameters(uint32 InGPUIndex)
 
 		FD3D12Buffer* IndexBuffer = FD3D12DynamicRHI::ResourceCast(Initializer.IndexBuffer.GetReference(), InGPUIndex);
 		const uint32 IndexStride = IndexBuffer ? IndexBuffer->GetStride() : 0;
+		const uint32 LocalIndicesPerPrimitive = FD3D12RayTracingGeometry::IndicesPerPrimitive;
 		for (int32 SegmentIndex = 0; SegmentIndex < Initializer.Segments.Num(); ++SegmentIndex)
 		{
 			const FRayTracingGeometrySegment& Segment = Initializer.Segments[SegmentIndex];
@@ -4126,7 +4115,7 @@ void FD3D12RayTracingGeometry::SetupHitGroupSystemParameters(uint32 InGPUIndex)
 				SystemParameters.VertexBuffer = VertexBuffer->ResourceLocation.GetGPUVirtualAddress() + Segment.VertexBufferOffset;
 			}
 
-			if (GeometryType == D3D12_RAYTRACING_GEOMETRY_TYPE_TRIANGLES && IndexBuffer != nullptr)
+			if (IndexBuffer != nullptr)
 			{
 				if (bBindless)
 				{
@@ -4136,7 +4125,7 @@ void FD3D12RayTracingGeometry::SetupHitGroupSystemParameters(uint32 InGPUIndex)
 				{
 					SystemParameters.IndexBuffer = IndexBuffer->ResourceLocation.GetGPUVirtualAddress();
 				}
-				SystemParameters.RootConstants.IndexBufferOffsetInBytes = Initializer.IndexBufferOffset + IndexStride * Segment.FirstPrimitive * FD3D12RayTracingGeometry::IndicesPerPrimitive;
+				SystemParameters.RootConstants.IndexBufferOffsetInBytes = Initializer.IndexBufferOffset + IndexStride * Segment.FirstPrimitive * LocalIndicesPerPrimitive;
 				SystemParameters.RootConstants.FirstPrimitive = Segment.FirstPrimitive;
 			}
 
@@ -4159,22 +4148,25 @@ void FD3D12RayTracingGeometry::CreateAccelerationStructureBuildDesc(FD3D12Comman
 	const uint32 IndexStride = Initializer.IndexBuffer ? Initializer.IndexBuffer->GetStride() : 0;
 	const bool bIsUpdate = BuildMode == EAccelerationStructureBuildMode::Update;
 
-	// Use the pre-built descs as template and set the GPU resource pointers (current VB/IB).
-	check(OutGeometryDescs.Num() == GeometryDescs.Num());
+	// Make sure we have the right amount of space reserved
+	check(OutGeometryDescs.Num() == Initializer.Segments.Num());
 	checkf(BuffersValid(GPUIndex), TEXT("Index & vertex buffers should be valid (not streamed out) when building the acceleration structure"));
 
 	FD3D12Buffer* IndexBuffer = CommandContext.RetrieveObject<FD3D12Buffer>(Initializer.IndexBuffer.GetReference());
-	FD3D12Buffer* NullTransformBufferD3D12 = CommandContext.RetrieveObject<FD3D12Buffer>(NullTransformBuffer.GetReference());
 		
 	const TArray<FD3D12HitGroupSystemParameters>& HitGroupSystemParametersForThisGPU = HitGroupSystemParameters[GPUIndex];
 	check(HitGroupSystemParametersForThisGPU.Num() == Initializer.Segments.Num());
 
 	ERayTracingAccelerationStructureFlags BuildFlags = GetRayTracingAccelerationStructureBuildFlags(Initializer);
 	D3D12_RAYTRACING_GEOMETRY_TYPE GeometryType = TranslateRayTracingGeometryType(Initializer.GeometryType);
+
+	// Fill in the GeometryDesc
+	TranslateRayTracingGeometryDescs(Initializer, OutGeometryDescs);
+
+	// second pass to actually set the current VB/IB buffer addresses
 	for (int32 SegmentIndex = 0; SegmentIndex < Initializer.Segments.Num(); ++SegmentIndex)
 	{
 		D3D12_RAYTRACING_GEOMETRY_DESC& Desc = OutGeometryDescs[SegmentIndex];
-		Desc = GeometryDescs[SegmentIndex]; // Copy from template
 
 		const FRayTracingGeometrySegment& Segment = Initializer.Segments[SegmentIndex];
 		const FD3D12HitGroupSystemParameters& SystemParameters = HitGroupSystemParametersForThisGPU[SegmentIndex];
@@ -4762,7 +4754,7 @@ void FD3D12CommandContext::RHIBuildAccelerationStructures(TConstArrayView<FRayTr
 				GRHIRayTracingAccelerationStructureAlignment);
 
 			// We need to keep D3D12_RAYTRACING_GEOMETRY_DESCs that are used in D3D12_BUILD_RAYTRACING_ACCELERATION_STRUCTURE_DESC alive.
-			const uint32 NumGeometryDescs = Geometry->GeometryDescs.Num();
+			const uint32 NumGeometryDescs = Geometry->GetInitializer().Segments.Num();
 			D3D12_RAYTRACING_GEOMETRY_DESC* LocalGeometryDescsMemory = (D3D12_RAYTRACING_GEOMETRY_DESC*)FMemStack::Get().Alloc(NumGeometryDescs * sizeof(D3D12_RAYTRACING_GEOMETRY_DESC), alignof(D3D12_RAYTRACING_GEOMETRY_DESC));
 			TArrayView<D3D12_RAYTRACING_GEOMETRY_DESC> LocalGeometryDescs = MakeArrayView(LocalGeometryDescsMemory, NumGeometryDescs);
 
@@ -4779,7 +4771,7 @@ void FD3D12CommandContext::RHIBuildAccelerationStructures(TConstArrayView<FRayTr
 			{
 				INC_DWORD_STAT(STAT_D3D12RayTracingBuiltBLAS);
 			}
-		}		
+		}
 	}
 
 	if (ShouldRunRayTracingGPUValidation())
diff --git a/Engine/Source/Runtime/D3D12RHI/Private/D3D12RayTracing.h b/Engine/Source/Runtime/D3D12RHI/Private/D3D12RayTracing.h
--- a/Engine/Source/Runtime/D3D12RHI/Private/D3D12RayTracing.h
+++ b/Engine/Source/Runtime/D3D12RHI/Private/D3D12RayTracing.h
@@ -182,8 +182,6 @@ public:
 
 	static constexpr uint32 IndicesPerPrimitive = 3; // Triangle geometry only
 
-	static FBufferRHIRef NullTransformBuffer; // Null transform for hidden sections
-
 	TRefCountPtr<FD3D12Buffer> AccelerationStructureBuffers[MAX_NUM_GPUS];
 
 	bool bRegisteredAsRenameListener[MAX_NUM_GPUS];
@@ -199,11 +197,6 @@ public:
 	FDebugName DebugName;
 	FName OwnerName;		// Store the path name of the owner object for resource tracking
 
-	// Array of geometry descriptions, one per segment (single-segment geometry is a common case).
-	// Only references CPU-accessible structures (no GPU resources).
-	// Used as a template for BuildAccelerationStructure() later.
-	TArray<D3D12_RAYTRACING_GEOMETRY_DESC, TInlineAllocator<1>> GeometryDescs;
-
 	uint64 AccelerationStructureCompactedSize = 0;
 		
 public:
