Make GPU driver checks respect adapter generations and update driver versions

#Synced-CL 51078227

diff --git a/Engine/Config/BaseHardware.ini b/Engine/Config/BaseHardware.ini
--- a/Engine/Config/BaseHardware.ini
+++ b/Engine/Config/BaseHardware.ini
@@ -36,8 +36,8 @@
 [GPU_NVIDIA Windows]
 SuggestedDriverVersion="560.81"
 
-; The newest driver which supports legacy Nvidia GPUs (700 series and older) is 474.64. We allow this driver for DX11.
-+DriverDenyList=(DriverVersion="<474.64", Reason="These driver versions have known stability issues and missing features", RHI="D3D11")
+; The newest driver which supports legacy Nvidia GPUs (700 series and older) is 475.14. We allow this driver for DX11.
++DriverDenyList=(DriverVersion="<475.14", Reason="These driver versions have known stability issues and missing features", RHI="D3D11")
 
 ; DX12 require recent drivers. The 536 series was released in June 2023.
 +DriverDenyList=(DriverVersion="<536.00", Reason="These driver versions have known stability issues and missing features", RHI="D3D12")
@@ -59,7 +59,7 @@ SuggestedDriverVersion="570.86"
 SuggestedDriverVersion="24.3.1"
 
 ; The newest driver which supports legacy AMD GPUs (GCN1-GCN3 chips like the Radeon R9 or older) is 22.6.1. We allow this driver for DX11.
-+DriverDenyList=(DriverVersion="<27.20.20913.2000", Reason="These driver versions have known stability issues and missing features", RHI="D3D11")
++DriverDenyList=(DriverVersion="<27.20.20912.1002", Reason="These driver versions have known stability issues and missing features", RHI="D3D11")
 
 ; The minimum version for DX12 is 24.3.1, published in March 2024. There are separate packages for different GPU generations and they
 ; have different internal versions, so we can't check for a version number explicitly (the RDNA driver is 31.0.24027.1012, older GPUs have
diff --git a/Engine/Source/Runtime/Core/Public/GenericPlatform/GenericPlatformDriver.h b/Engine/Source/Runtime/Core/Public/GenericPlatform/GenericPlatformDriver.h
--- a/Engine/Source/Runtime/Core/Public/GenericPlatform/GenericPlatformDriver.h
+++ b/Engine/Source/Runtime/Core/Public/GenericPlatform/GenericPlatformDriver.h
@@ -157,6 +157,9 @@ struct FGPUDriverInfo
 	// Hardware vendor ID, e.g. 0x10DE for NVIDIA.
 	uint32 VendorId = 0;
 
+	// Hardware adapter generation (vendor-specific).
+	uint32 AdapterGeneration = 0;
+
 	// Name of the graphics device, e.g. "NVIDIA GeForce GTX 680" or "AMD Radeon R9 200 / HD 7900 Series".
 	FString DeviceDescription;
 
@@ -198,6 +201,11 @@ struct FGPUDriverInfo
 		return VendorId == 0x5143 || VendorId == 0x4D4F4351;
 	}
 
+	bool IsAdapterGenerationKnown() const
+	{
+		return AdapterGeneration != 0;
+	}
+
 	bool IsSameDriverVersionGeneration(const FDriverVersion& UnifiedDriverVersion) const
 	{
 		if (IsIntel())
@@ -275,9 +283,25 @@ static FDateTime ParseMonthDayYearDate(const FString& DateString)
 // Represents an optional set of constraints that a driver configuration entry can include:
 //   * The RHI being used (e.g. "D3D12" will result in the entry to be considered only when running in D3D12).
 //   * The adapter name, using a regular expression (e.g. ".*RTX.*" will result in the entry to be considered only for cards that contain the string "RTX").
+//   * The adapter generation. This is a vendor-specific integer and may not be implemented for all vendors.
 struct FDriverConfigEntryConstraints {
+	struct FIntegerConstraint {
+		EComparisonOp ComparisonOp;
+		uint32 Value;
+
+		static FIntegerConstraint FromString(const FString& IntegerConstraintString)
+		{
+			FIntegerConstraint Constraint;
+			const TCHAR* Chars = *IntegerConstraintString;
+			Constraint.ComparisonOp = ParseComparisonOp(Chars);
+			Constraint.Value = FCString::Atoi(Chars);
+			return Constraint;
+		}
+	};
+
 	TOptional<FString> RHINameConstraint;
 	TOptional<FRegexPattern> AdapterNameRegexConstraint;
+	TOptional<FIntegerConstraint> AdapterGenerationConstraint;
 
 	FDriverConfigEntryConstraints() = default;
 
@@ -296,6 +320,13 @@ struct FDriverConfigEntryConstraints {
 		{
 			AdapterNameRegexConstraint.Emplace(AdapterNameRegexString);
 		}
+
+		FString AdapterGenerationString;
+		FParse::Value(Entry, TEXT("AdapterGeneration="), AdapterGenerationString);
+		if (!AdapterGenerationString.IsEmpty())
+		{
+			AdapterGenerationConstraint = FIntegerConstraint::FromString(AdapterGenerationString);
+		}
 	}
 
 	bool AreConstraintsSatisfied(const FGPUDriverInfo& Info, uint32& OutNumSatisfiedConstraints) const
@@ -321,12 +352,21 @@ struct FDriverConfigEntryConstraints {
 			OutNumSatisfiedConstraints++;
 		}
 
+		if (AdapterGenerationConstraint)
+		{
+			if (!Info.IsAdapterGenerationKnown() || !CompareWithOp(Info.AdapterGeneration, AdapterGenerationConstraint->ComparisonOp, AdapterGenerationConstraint->Value))
+			{
+				return false;
+			}
+			OutNumSatisfiedConstraints++;
+		}
+
 		return true;
 	}
 
 	bool HasConstraints()
 	{
-		return RHINameConstraint || AdapterNameRegexConstraint;
+		return RHINameConstraint || AdapterNameRegexConstraint || AdapterGenerationConstraint;
 	}
 };
 
@@ -483,6 +523,37 @@ struct FSuggestedDriverEntry : public FDriverConfigEntryConstraints {
 	FString SuggestedDriverVersion;
 };
 
+// This corresponds to one SuggestedDriverList entry in the Hardware.ini configuration file. One entry includes:
+//   * A driver version
+//   * A set of optional constraints that the entry will consider when the driver is being tested (see FDriverConfigEntryConstraints).
+struct FSuggestedDriverListEntry : public FDriverConfigEntryConstraints {
+	FSuggestedDriverListEntry() = default;
+
+	FSuggestedDriverListEntry(const TCHAR* Entry)
+		: FDriverConfigEntryConstraints(Entry)
+	{
+		FParse::Value(Entry, TEXT("DriverVersion="), DriverVersion);
+
+		ensureMsgf(IsValid(), TEXT("DriverVersion must be specified in a suggested driver list entry"));
+	}
+
+	// Checks whether this entry applies to the given driver info, and also returns the number
+	// of constraints that have been satisfied in OutNumSatisfiedConstraints.
+	bool AppliesToDriver(const FGPUDriverInfo& Info, uint32& OutNumSatisfiedConstraints) const
+	{
+		OutNumSatisfiedConstraints = 0;
+
+		return IsValid() && AreConstraintsSatisfied(Info, OutNumSatisfiedConstraints);
+	}
+
+	bool IsValid() const
+	{
+		return !DriverVersion.IsEmpty();
+	}
+
+	FString	DriverVersion;
+};
+
 class FGPUDriverHelper
 {
 public:
@@ -492,9 +563,21 @@ public:
 	}
 
 	// Returns the best suggested driver version that matches the current configuration, if it exists. 
-	TOptional<FSuggestedDriverEntry> FindSuggestedDriverVersion() const
+	FString FindSuggestedDriverVersion() const
 	{
-		return GetBestConfigEntry<FSuggestedDriverEntry>(TEXT("SuggestedDriverVersion"));
+		TOptional<FSuggestedDriverListEntry> ListEntry = GetBestConfigEntry<FSuggestedDriverListEntry>(TEXT("SuggestedDriverList"));
+		if (ListEntry.IsSet())
+		{
+			return ListEntry->DriverVersion;
+		}
+
+		TOptional<FSuggestedDriverEntry> Entry = GetBestConfigEntry<FSuggestedDriverEntry>(TEXT("SuggestedDriverVersion"));
+		if (Entry.IsSet())
+		{
+			return Entry->SuggestedDriverVersion;
+		}
+
+		return FString();
 	}
 
 	// Finds the best driver denylist entry that matches the current driver and configuration, if it exists. 
diff --git a/Engine/Source/Runtime/Core/Tests/GenericPlatform/GenericPlatformDriverTest.cpp b/Engine/Source/Runtime/Core/Tests/GenericPlatform/GenericPlatformDriverTest.cpp
--- a/Engine/Source/Runtime/Core/Tests/GenericPlatform/GenericPlatformDriverTest.cpp
+++ b/Engine/Source/Runtime/Core/Tests/GenericPlatform/GenericPlatformDriverTest.cpp
@@ -152,6 +152,7 @@ TEST_CASE_NAMED(DriverVersionTest, "System::Core::GenericPlatform::Driver", "[Ap
 			CHECK(Entry.AdapterNameRegexConstraint);
 			CHECK(Entry.DriverVersion->ComparisonOp == EComparisonOp::ECO_Equal);
 			CHECK(Entry.DriverVersion->Version == FDriverVersion(MakeVersionString({ 301, 12 })));
+			CHECK(!Entry.AppliesToLatestDrivers());
 		}
 		{
 			FDriverDenyListEntry Entry(TEXT("(DriverVersion=\"<=311.16\", Reason=\"Test with no constraints\")"));
@@ -159,6 +160,7 @@ TEST_CASE_NAMED(DriverVersionTest, "System::Core::GenericPlatform::Driver", "[Ap
 			CHECK(Entry.DenylistReason == TEXT("Test with no constraints"));
 			CHECK(Entry.DriverVersion->ComparisonOp == EComparisonOp::ECO_LessOrEqual);
 			CHECK(Entry.DriverVersion->Version == FDriverVersion(MakeVersionString({ 311, 16 })));
+			CHECK(!Entry.AppliesToLatestDrivers());
 		}
 		{
 			FDriverDenyListEntry Entry(TEXT("(DriverDate=\">=12-31-2021\", RHI=\"D3D11\", Reason=\"Test date\")"));
@@ -169,6 +171,16 @@ TEST_CASE_NAMED(DriverVersionTest, "System::Core::GenericPlatform::Driver", "[Ap
 			CHECK(Entry.DriverDate->Date == FDateTime(2021, 12, 31));
 			CHECK(Entry.AppliesToLatestDrivers());
 		}
+		{
+			FDriverDenyListEntry Entry(TEXT("(DriverVersion=\"<=311.16\", AdapterGeneration=\"<=5\", Reason=\"Test generation\")"));
+			CHECK(Entry.IsValid());
+			CHECK(Entry.DenylistReason == TEXT("Test generation"));
+			CHECK(Entry.AdapterGenerationConstraint->ComparisonOp == EComparisonOp::ECO_LessOrEqual);
+			CHECK(Entry.AdapterGenerationConstraint->Value == 5);
+			CHECK(Entry.DriverVersion->ComparisonOp == EComparisonOp::ECO_LessOrEqual);
+			CHECK(Entry.DriverVersion->Version == FDriverVersion(MakeVersionString({ 311, 16 })));
+			CHECK(!Entry.AppliesToLatestDrivers());
+		}
 		{
 			FSuggestedDriverEntry Entry(TEXT("(DriverVersion=\"512.15\", RHI=\"D3D11\", AdapterNameRegex=\".*4050.*\")"));
 			CHECK(Entry.IsValid());
diff --git a/Engine/Source/Runtime/D3D12RHI/Private/D3D12Adapter.cpp b/Engine/Source/Runtime/D3D12RHI/Private/D3D12Adapter.cpp
--- a/Engine/Source/Runtime/D3D12RHI/Private/D3D12Adapter.cpp
+++ b/Engine/Source/Runtime/D3D12RHI/Private/D3D12Adapter.cpp
@@ -708,7 +708,7 @@ void FD3D12Adapter::CreateRootDevice(bool bWithDebug)
 		));
 
 		INTCExtensionInfo INTCExtensionInfo{};
-		if (INTCExtensionContext* IntelExtensionContext = CreateIntelExtensionsContext(Device, INTCExtensionInfo))
+		if (INTCExtensionContext* IntelExtensionContext = CreateIntelExtensionsContext(Device, INTCExtensionInfo, GetD3DAdapterDesc().DeviceId))
 		{
 			EnableIntelAtomic64Support(IntelExtensionContext, INTCExtensionInfo);
 			// Destroy the context to release all reference to ID3D12Device
diff --git a/Engine/Source/Runtime/D3D12RHI/Private/D3D12RHI.cpp b/Engine/Source/Runtime/D3D12RHI/Private/D3D12RHI.cpp
--- a/Engine/Source/Runtime/D3D12RHI/Private/D3D12RHI.cpp
+++ b/Engine/Source/Runtime/D3D12RHI/Private/D3D12RHI.cpp
@@ -252,6 +252,8 @@ FD3D12DynamicRHI::FD3D12DynamicRHI(TSharedPtr<FD3D12Adapter>& InChosenAdapter, b
 	GRHIGlobals.SupportsVertexShaderUAVs = true;
 
 	GRHIGlobals.NeedsExtraTransitions = true;
+
+	InitializeDeviceAgnosticVendorAPI();
 }
 
 void FD3D12DynamicRHI::PostInit()
diff --git a/Engine/Source/Runtime/D3D12RHI/Private/D3D12RHIPrivate.h b/Engine/Source/Runtime/D3D12RHI/Private/D3D12RHIPrivate.h
--- a/Engine/Source/Runtime/D3D12RHI/Private/D3D12RHIPrivate.h
+++ b/Engine/Source/Runtime/D3D12RHI/Private/D3D12RHIPrivate.h
@@ -577,6 +577,8 @@ protected:
 	AGSContext* AmdAgsContext;
 	uint32 AmdSupportedExtensionFlags;
 
+	bool bNvAPIInitialized = false;
+
 	INTCExtensionContext* IntelExtensionContext = nullptr;
 
 #if WITH_NVAPI_RAYTRACING_VALIDATION
@@ -688,6 +690,8 @@ protected:
 	FDisplayInformationArray DisplayList;
 
 	void ProcessDeferredDeletionQueue_Platform();
+
+	void InitializeDeviceAgnosticVendorAPI();
 };
 
 ENUM_CLASS_FLAGS(FD3D12DynamicRHI::EQueueStatus);
diff --git a/Engine/Source/Runtime/D3D12RHI/Private/Windows/WindowsD3D12Device.cpp b/Engine/Source/Runtime/D3D12RHI/Private/Windows/WindowsD3D12Device.cpp
--- a/Engine/Source/Runtime/D3D12RHI/Private/Windows/WindowsD3D12Device.cpp
+++ b/Engine/Source/Runtime/D3D12RHI/Private/Windows/WindowsD3D12Device.cpp
@@ -33,6 +33,16 @@ extern bool D3D12RHI_AllowSoftwareFallback();
 extern bool D3D12RHI_ShouldAllowAsyncResourceCreation();
 extern bool D3D12RHI_ShouldForceCompatibility();
 
+#if WITH_AMD_AGS
+// Filled in during InitializeDeviceAgnosticVendorAPI if IsRHIDeviceAMD
+struct AmdAgsInfo
+{
+	AGSContext* AmdAgsContext;
+	AGSGPUInfo AmdGpuInfo;
+};
+static AmdAgsInfo AmdInfo;
+#endif
+
 #if INTEL_EXTENSIONS
 bool GDX12INTCAtomicUInt64Emulation = false;
 #endif
@@ -322,7 +332,6 @@ INTCExtensionContext* CreateIntelExtensionsContext(ID3D12Device* Device, INTCExt
 		}
 	}
 
-
 	return IntelExtensionContext;
 }
 
@@ -1309,6 +1318,107 @@ static void ClearPSODriverCache()
 	}
 }
 
+void FD3D12DynamicRHI::InitializeDeviceAgnosticVendorAPI()
+{
+	// We set up AMD and NVIDIA APIs here since they don't
+	// need a device. Intel extensions are enabled after
+	// device creation.
+	if (!UE::RHICore::AllowVendorDevice())
+	{
+		return;
+	}
+
+	const DXGI_ADAPTER_DESC& AdapterDesc = GetAdapter().GetD3DAdapterDesc();
+
+#ifdef AMD_AGS_API
+	if (IsRHIDeviceAMD())
+	{
+		check(AmdAgsContext == nullptr);
+
+		// agsInitialize should be called before D3D device creation.
+		if (agsInitialize(AGS_MAKE_VERSION(AMD_AGS_VERSION_MAJOR, AMD_AGS_VERSION_MINOR, AMD_AGS_VERSION_PATCH), nullptr, &AmdAgsContext, &AmdInfo.AmdGpuInfo) == AGS_SUCCESS)
+		{
+			AmdInfo.AmdAgsContext = AmdAgsContext;
+			for (int32 DeviceIndex = 0; DeviceIndex < AmdInfo.AmdGpuInfo.numDevices; DeviceIndex++)
+			{
+				const AGSDeviceInfo& DeviceInfo = AmdInfo.AmdGpuInfo.devices[DeviceIndex];
+				if (AdapterDesc.VendorId == DeviceInfo.vendorId && AdapterDesc.DeviceId == DeviceInfo.deviceId)
+				{
+					const uint32 AsicFamily = DeviceInfo.asicFamily;
+					UE_LOG(LogD3D12RHI, Log, TEXT("AMD asicFamily=%u"), AsicFamily);
+					GRHIDeviceIsAMDPreGCNArchitecture = DeviceInfo.asicFamily == AGSDeviceInfo::AsicFamily_PreGCN;
+					GRHIGlobals.GpuInfo.AdapterGeneration = AsicFamily;
+				}
+			}
+
+			if (GRHIDeviceIsAMDPreGCNArchitecture)
+			{
+				UE_LOG(LogD3D12RHI, Log, TEXT("AMD Pre GCN architecture detected, some driver workarounds will be in place"));
+			}
+		}
+		else
+		{
+			UE_LOG(LogD3D12RHI, Warning, TEXT("Failed to initialize AMD AGS API"));
+
+			FMemory::Memzero(&AmdInfo, sizeof(AmdInfo));
+			// If agsInitialize returns anything but AGS_SUCCESS, the context pointer should be
+			// guaranteed to be NULL, but we'll set it here explicitly, just to be safe.
+			AmdAgsContext = nullptr;
+		}
+	}
+	else
+	{
+		FMemory::Memzero(&AmdInfo, sizeof(AmdInfo));
+	}
+#endif // AMD_AGS_API
+
+#ifdef NVAPI_INTERFACE
+	if (IsRHIDeviceNVIDIA())
+	{
+		if (NvAPI_Initialize() == NVAPI_OK)
+		{
+			bNvAPIInitialized = true;
+
+			NvPhysicalGpuHandle GPUHandles[NVAPI_MAX_PHYSICAL_GPUS];
+			NvU32 NumGPUs = 0;
+			if (NvAPI_EnumPhysicalGPUs(GPUHandles, &NumGPUs) == NVAPI_OK)
+			{
+				for (NvU32 GPU = 0; GPU < NumGPUs; ++GPU)
+				{
+					const NvPhysicalGpuHandle& GPUHandle = GPUHandles[GPU];
+
+					NvU32 Unused = 0, SubsystemId = 0, RevisionId = 0, ExtDeviceId = 0;
+					if (NvAPI_GPU_GetPCIIdentifiers(GPUHandle, &Unused, &SubsystemId, &RevisionId, &ExtDeviceId) != NVAPI_OK)
+					{
+						continue;
+					}
+
+					if (ExtDeviceId != AdapterDesc.DeviceId || SubsystemId != AdapterDesc.SubSysId || RevisionId != AdapterDesc.Revision)
+					{
+						continue;
+					}
+
+					NV_GPU_ARCH_INFO ArchInfo;
+					ArchInfo.version = NV_GPU_ARCH_INFO_VER_2;
+					if (NvAPI_GPU_GetArchInfo(GPUHandle, &ArchInfo) != NVAPI_OK)
+					{
+						continue;
+					}
+
+					const uint32 Architecture = ArchInfo.architecture;
+					UE_LOG(LogD3D12RHI, Log, TEXT("NVIDIA Architecture=%u"), Architecture);
+					GRHIGlobals.GpuInfo.AdapterGeneration = Architecture;
+				}
+			}
+		}
+		else
+		{
+			UE_LOG(LogD3D12RHI, Warning, TEXT("Failed to initialize NVAPI"));
+		}
+	}
+#endif // NVAPI_INTERFACE
+}
+
 void FD3D12DynamicRHI::Init()
 {
 #if D3D12RHI_SUPPORTS_WIN_PIX
@@ -1334,24 +1444,10 @@ void FD3D12DynamicRHI::Init()
 		ClearPSODriverCache();
 	}
 
-
 	const DXGI_ADAPTER_DESC& AdapterDesc = GetAdapter().GetD3DAdapterDesc();
 
 	UE_LOG(LogD3D12RHI, Log, TEXT("    GPU DeviceId: 0x%x (for the marketing name, search the web for \"GPU Device Id\")"), AdapterDesc.DeviceId);
 
-#if WITH_AMD_AGS
-	// Initialize the AMD AGS utility library, when running on an AMD device
-	AGSGPUInfo AmdAgsGpuInfo = {};
-	if (IsRHIDeviceAMD() && UE::RHICore::AllowVendorDevice())
-	{
-		check(AmdAgsContext == nullptr);
-		check(AmdSupportedExtensionFlags == 0);
-
-		// agsInit should be called before D3D device creation
-		agsInitialize(AGS_MAKE_VERSION(AMD_AGS_VERSION_MAJOR, AMD_AGS_VERSION_MINOR, AMD_AGS_VERSION_PATCH), nullptr, &AmdAgsContext, &AmdAgsGpuInfo);
-	}
-#endif
-
 	// Create a device chain for the adapter we have chosen. This could be a single discrete card,
 	// a set discrete cards linked together (i.e. SLI/Crossfire) an Integrated device or any combination of the above
 	if (FD3D12Adapter* const Adapter = ChosenAdapter.Get())
@@ -1390,13 +1486,12 @@ void FD3D12DynamicRHI::Init()
 	bool bHasVendorSupportForAtomic64 = false;
 
 #if WITH_AMD_AGS
-
 	// Check if the AMD device is pre-RDNA, and ensure it doesn't misreport wave32 support
 	if (IsRHIDeviceAMD() && UE::RHICore::AllowVendorDevice() && AmdAgsContext)
 	{
-		for (int32 DeviceIndex = 0; DeviceIndex < AmdAgsGpuInfo.numDevices; DeviceIndex++)
+		for (int32 DeviceIndex = 0; DeviceIndex < AmdInfo.AmdGpuInfo.numDevices; DeviceIndex++)
 		{
-			const AGSDeviceInfo& DeviceInfo = AmdAgsGpuInfo.devices[DeviceIndex];
+			const AGSDeviceInfo& DeviceInfo = AmdInfo.AmdGpuInfo.devices[DeviceIndex];
 			if (DeviceInfo.deviceId == AdapterDesc.DeviceId && DeviceInfo.vendorId == AdapterDesc.VendorId)
 			{
 				if (DeviceInfo.asicFamily != AGSDeviceInfo::AsicFamily_Unknown && DeviceInfo.asicFamily < AGSDeviceInfo::AsicFamily_RDNA)
@@ -1444,13 +1539,11 @@ void FD3D12DynamicRHI::Init()
 		UE_LOG(LogD3D12RHI, Warning, TEXT("Ray tracing is disabled because it requires Windows 10 version %u"), (uint32)GMinimumWindowsBuildVersionForRayTracing);
 	}
 
-	const FDriverVersion CurrentDriverVersion(GRHIAdapterUserDriverVersion);
+	const FDriverVersion CurrentDriverVersion(GRHIGlobals.GpuInfo.AdapterUserDriverVersion);
 
 #if WITH_NVAPI
 	if (IsRHIDeviceNVIDIA() && UE::RHICore::AllowVendorDevice())
 	{
-		const NvAPI_Status NvStatus = NvAPI_Initialize();
-
 		NvU32 PipelineFlags = NVAPI_D3D12_PIPELINE_CREATION_STATE_FLAGS_NONE;
 
 		// Helper to allow all pipeline state capabilities to test and accumulate their support, so all can be applied atomically
@@ -1468,7 +1561,7 @@ void FD3D12DynamicRHI::Init()
 				return false;
 			};
 
-		if (NvStatus == NVAPI_OK)
+		if (bNvAPIInitialized)
 		{
 			NvAPI_Status NvStatusAtomicU64 = NvAPI_D3D12_IsNvShaderExtnOpCodeSupported(GetAdapter().GetD3DDevice(), NV_EXTN_OP_UINT64_ATOMIC, &bHasVendorSupportForAtomic64);
 			if (NvStatusAtomicU64 != NVAPI_OK)
@@ -1618,10 +1711,6 @@ void FD3D12DynamicRHI::Init()
 				UE_LOG(LogD3D12RHI, Log, TEXT("Skipped NVAPI RT queries since we the feature level is below SM6"));
 			}
 		}
-		else
-		{
-			UE_LOG(LogD3D12RHI, Warning, TEXT("Failed to initialize NVAPI"));
-		}
 
 		// Disable ray tracing for old Nvidia drivers
 		if (GRHISupportsRayTracing 
diff --git a/Engine/Source/Runtime/RHI/Private/DynamicRHI.cpp b/Engine/Source/Runtime/RHI/Private/DynamicRHI.cpp
--- a/Engine/Source/Runtime/RHI/Private/DynamicRHI.cpp
+++ b/Engine/Source/Runtime/RHI/Private/DynamicRHI.cpp
@@ -85,12 +85,15 @@ void RHIDetectAndWarnOfBadDrivers(bool bHasEditorToken)
 
 	if (GRHIVendorId == 0)
 	{
-		UE_LOG(LogRHI, Log, TEXT("Skipping Driver Check, no vendor ID set."));
+		UE_LOG(LogRHI, Log, TEXT("Skipping driver check, no vendor ID set."));
 		return;
 	}
 
+	UE_LOG(LogRHI, Log, TEXT("Performing driver check for denylisted drivers."));
+
 	FGPUDriverInfo DriverInfo;
 	DriverInfo.VendorId = GRHIVendorId;
+	DriverInfo.AdapterGeneration = GRHIGlobals.GpuInfo.AdapterGeneration;
 	DriverInfo.DeviceDescription = GRHIAdapterName;
 	DriverInfo.ProviderName = TEXT("Unknown");
 	DriverInfo.InternalDriverVersion = GRHIAdapterInternalDriverVersion;
@@ -98,28 +101,11 @@ void RHIDetectAndWarnOfBadDrivers(bool bHasEditorToken)
 	DriverInfo.DriverDate = GRHIAdapterDriverDate;
 	DriverInfo.RHIName = GDynamicRHI ? GDynamicRHI->GetName() : FString();
 
-	FGPUDriverHelper DetectedGPUHardware(DriverInfo);
-
 	// Pre-GCN GPUs usually don't support updating to latest driver
 	// But it is unclear what is the latest version supported as it varies from card to card
 	// So just don't complain if pre-gcn
 	if (DriverInfo.IsValid() && !GRHIDeviceIsAMDPreGCNArchitecture)
 	{
-		TOptional<FDriverDenyListEntry> DenyListEntry = DetectedGPUHardware.FindDriverDenyListEntry();
-
-		GRHIAdapterDriverOnDenyList = DenyListEntry.IsSet() && DenyListEntry->IsValid();
-		FGenericCrashContext::SetEngineData(TEXT("RHI.DriverDenylisted"), GRHIAdapterDriverOnDenyList ? TEXT("true") : TEXT("false"));
-
-		if(!GRHIAdapterDriverOnDenyList)
-		{
-			return;
-		}
-
-		TOptional<FSuggestedDriverEntry> SuggestedDriver = DetectedGPUHardware.FindSuggestedDriverVersion();
-		const FString SuggestedDriverString = SuggestedDriver ? SuggestedDriver->SuggestedDriverVersion : FString("Unknown");
-
-		UE_LOG(LogRHI, Warning, TEXT("Out of date driver found. Using: '%s' Suggested: '%s'"), *DriverInfo.UserDriverVersion, (SuggestedDriver ? *SuggestedDriver->SuggestedDriverVersion : TEXT("Unknown")));
-
 		TArray<FString> DeviceCanUpdateDriverList;
 		GConfig->GetArray(TEXT("Devices"), TEXT("DeviceCanUpdateDriverList"), DeviceCanUpdateDriverList, GHardwareIni);
 
@@ -142,11 +128,32 @@ void RHIDetectAndWarnOfBadDrivers(bool bHasEditorToken)
 			if (DriverInfo.VendorId && GRHIDeviceId &&
 				DriverInfo.VendorId == VendorIdInt && GRHIDeviceId == DeviceIdInt)
 			{
+				int32 DriverGeneration = 0;
+				if (FParse::Value(Line + 1, TEXT("DriverGeneration="), DriverGeneration) && (DriverGeneration > 0))
+				{
+					// Allow config to override vendor API queries.
+					DriverInfo.AdapterGeneration = DriverGeneration;
+				}
+
 				bDeviceCanUpdateDriver = true;
 				break;
 			}
 		}
 
+		FGPUDriverHelper DetectedGPUHardware(DriverInfo);
+		TOptional<FDriverDenyListEntry> DenyListEntry = DetectedGPUHardware.FindDriverDenyListEntry();
+
+		GRHIAdapterDriverOnDenyList = DenyListEntry.IsSet() && DenyListEntry->IsValid();
+		FGenericCrashContext::SetEngineData(TEXT("RHI.DriverDenylisted"), GRHIAdapterDriverOnDenyList ? TEXT("true") : TEXT("false"));
+
+		if(!GRHIAdapterDriverOnDenyList)
+		{
+			return;
+		}
+
+		const FString SuggestedDriverString = DetectedGPUHardware.FindSuggestedDriverVersion();
+		UE_LOG(LogRHI, Warning, TEXT("Out of date driver found. Using: '%s' Suggested: '%s'"), *DriverInfo.UserDriverVersion, (SuggestedDriverString.IsEmpty() ? TEXT("Unknown") : *SuggestedDriverString));
+
 		// Only alert users who are capable of updating their driver. Assume vendors with an empty list can always update.
 		// The warning message can also be suppressed with r.WarnOfBadDrivers=0.
 		bool bShowPrompt = bDeviceCanUpdateDriver || !bVendorHasEntries;
@@ -154,6 +161,8 @@ void RHIDetectAndWarnOfBadDrivers(bool bHasEditorToken)
 
 		if (bShowPrompt)
 		{
+			UE_LOG(LogRHI, Log, TEXT("Current driver is denylisted."));
+
 			// Note: we don't localize the vendor's name.
 			FString VendorString = DriverInfo.ProviderName;
 			FText HyperlinkText;
@@ -185,16 +194,16 @@ void RHIDetectAndWarnOfBadDrivers(bool bHasEditorToken)
 			Args.Add(TEXT("InstalledVer"), FText::FromString(DriverInfo.UserDriverVersion));
 
 			// Find the best driver version to recommend.
-			if (SuggestedDriver)
+			if (!SuggestedDriverString.IsEmpty())
 			{
 				// Suggest the latest too, if not denylisted.
 				if (!DenyListEntry->AppliesToLatestDrivers())
 				{
-					Args.Add(TEXT("RecommendedVer"), FText::Format(NSLOCTEXT("MessageDialog", "SuggestedDriverOrLatest", "{0} or latest driver available"), FText::FromString(SuggestedDriver->SuggestedDriverVersion)));
+					Args.Add(TEXT("RecommendedVer"), FText::Format(NSLOCTEXT("MessageDialog", "SuggestedDriverOrLatest", "{0} or latest driver available"), FText::FromString(SuggestedDriverString)));
 				}
 				else
 				{
-					Args.Add(TEXT("RecommendedVer"), FText::FromString(SuggestedDriver->SuggestedDriverVersion));
+					Args.Add(TEXT("RecommendedVer"), FText::FromString(SuggestedDriverString));
 				}
 			}
 			else
@@ -296,7 +305,6 @@ void RHIInit(bool bHasEditorToken)
 			if (GDynamicRHI)
 			{
 #if PLATFORM_WINDOWS || PLATFORM_MAC || PLATFORM_UNIX
-
 				// Get driver version. Creating GDynamicRHI is expected to set GRHIAdapterName.
 				FGPUDriverInfo GPUDriverInfo = FPlatformMisc::GetGPUDriverInfo(GRHIAdapterName);
 				if (GPUDriverInfo.IsValid())
diff --git a/Engine/Source/Runtime/RHI/Public/RHIGlobals.h b/Engine/Source/Runtime/RHI/Public/RHIGlobals.h
--- a/Engine/Source/Runtime/RHI/Public/RHIGlobals.h
+++ b/Engine/Source/Runtime/RHI/Public/RHIGlobals.h
@@ -135,6 +135,9 @@ struct FRHIGlobals
 
 		// 0 means not defined yet, use functions like IsRHIDeviceAMD() to access
 		uint32 VendorId = 0;
+
+		// Vendor-specific generation number of the GPU. 0 if unknown.
+		uint32 AdapterGeneration = 0;
 	} GpuInfo;
 
 	// true if the RHI supports Pixel Shader UAV
diff --git a/Engine/Source/Runtime/Windows/D3D11RHI/Private/D3D11Device.cpp b/Engine/Source/Runtime/Windows/D3D11RHI/Private/D3D11Device.cpp
--- a/Engine/Source/Runtime/Windows/D3D11RHI/Private/D3D11Device.cpp
+++ b/Engine/Source/Runtime/Windows/D3D11RHI/Private/D3D11Device.cpp
@@ -232,6 +232,8 @@ FD3D11DynamicRHI::FD3D11DynamicRHI(IDXGIFactory1* InDXGIFactory1, D3D_FEATURE_LE
 	}
 
 	StaticUniformBuffers.AddZeroed(FUniformBufferStaticSlotRegistry::Get().GetSlotCount());
+
+	InitializeDeviceAgnosticVendorAPI();
 }
 
 FD3D11DynamicRHI::~FD3D11DynamicRHI()
diff --git a/Engine/Source/Runtime/Windows/D3D11RHI/Private/D3D11RHIPrivate.h b/Engine/Source/Runtime/Windows/D3D11RHI/Private/D3D11RHIPrivate.h
--- a/Engine/Source/Runtime/Windows/D3D11RHI/Private/D3D11RHIPrivate.h
+++ b/Engine/Source/Runtime/Windows/D3D11RHI/Private/D3D11RHIPrivate.h
@@ -439,6 +439,8 @@ public:
 	}
 
 private:
+	void InitializeDeviceAgnosticVendorAPI();
+
 	D3D11RHI_API void EnableDepthBoundsTest(bool bEnable, float MinDepth, float MaxDepth);
 
 	static D3D11RHI_API void ClearUAV(TRHICommandList_RecursiveHazardous<FD3D11DynamicRHI>& RHICmdList, FD3D11UnorderedAccessView* UAV, const void* ClearValues, bool bFloat);
@@ -533,6 +535,8 @@ protected:
 	 */
 	AGSContext* AmdAgsContext;
 
+	bool bNvAPIInitialized = false;
+
 #if INTEL_EXTENSIONS
 	// Context and functions for Intel extension framework utility library
 	INTCExtensionContext* IntelExtensionContext;
diff --git a/Engine/Source/Runtime/Windows/D3D11RHI/Private/Windows/WindowsD3D11Device.cpp b/Engine/Source/Runtime/Windows/D3D11RHI/Private/Windows/WindowsD3D11Device.cpp
--- a/Engine/Source/Runtime/Windows/D3D11RHI/Private/Windows/WindowsD3D11Device.cpp
+++ b/Engine/Source/Runtime/Windows/D3D11RHI/Private/Windows/WindowsD3D11Device.cpp
@@ -35,7 +35,7 @@ FD3D11DynamicRHI*	GD3D11RHI = nullptr;
 bool D3D11RHI_ShouldAllowAsyncResourceCreation();
 
 #ifdef AMD_AGS_API
-// Filled in during InitD3DDevice if IsRHIDeviceAMD
+// Filled in during InitializeDeviceAgnosticVendorAPI if IsRHIDeviceAMD
 struct AmdAgsInfo
 {
 	AGSContext* AmdAgsContext;
@@ -1012,6 +1012,105 @@ void FD3D11DynamicRHI::StopIntelExtensions()
 }
 #endif // INTEL_EXTENSIONS
 
+void FD3D11DynamicRHI::InitializeDeviceAgnosticVendorAPI()
+{
+	// We set up AMD and NVIDIA APIs here since they don't
+	// need a device. Intel extensions are enabled after
+	// device creation.
+	if (!UE::RHICore::AllowVendorDevice())
+	{
+		return;
+	}
+
+#ifdef AMD_AGS_API
+	if (IsRHIDeviceAMD())
+	{
+		check(AmdAgsContext == nullptr);
+
+		// agsInitialize should be called before D3D device creation.
+		if (agsInitialize(AGS_MAKE_VERSION(AMD_AGS_VERSION_MAJOR, AMD_AGS_VERSION_MINOR, AMD_AGS_VERSION_PATCH), nullptr, &AmdAgsContext, &AmdInfo.AmdGpuInfo) == AGS_SUCCESS)
+		{
+			AmdInfo.AmdAgsContext = AmdAgsContext;
+			for (int32 DeviceIndex = 0; DeviceIndex < AmdInfo.AmdGpuInfo.numDevices; DeviceIndex++)
+			{
+				const AGSDeviceInfo& DeviceInfo = AmdInfo.AmdGpuInfo.devices[DeviceIndex];
+				if (Adapter.DXGIAdapterDesc.VendorId == DeviceInfo.vendorId && Adapter.DXGIAdapterDesc.DeviceId == DeviceInfo.deviceId)
+				{
+					const uint32 AsicFamily = DeviceInfo.asicFamily;
+					UE_LOG(LogD3D11RHI, Log, TEXT("AMD asicFamily=%u"), AsicFamily);
+					GRHIDeviceIsAMDPreGCNArchitecture = DeviceInfo.asicFamily == AGSDeviceInfo::AsicFamily_PreGCN;
+					GRHIGlobals.GpuInfo.AdapterGeneration = AsicFamily;
+				}
+			}
+
+			if (GRHIDeviceIsAMDPreGCNArchitecture)
+			{
+				UE_LOG(LogD3D11RHI, Log, TEXT("AMD Pre GCN architecture detected, some driver workarounds will be in place"));
+			}
+		}
+		else
+		{
+			UE_LOG(LogD3D11RHI, Warning, TEXT("Failed to initialize AMD AGS API"));
+
+			FMemory::Memzero(&AmdInfo, sizeof(AmdInfo));
+			// If agsInitialize returns anything but AGS_SUCCESS, the context pointer should be
+			// guaranteed to be NULL, but we'll set it here explicitly, just to be safe.
+			AmdAgsContext = nullptr;
+		}
+	}
+	else
+	{
+		FMemory::Memzero(&AmdInfo, sizeof(AmdInfo));
+	}
+#endif // AMD_AGS_API
+
+#ifdef NVAPI_INTERFACE
+	if (IsRHIDeviceNVIDIA())
+	{
+		if (NvAPI_Initialize() == NVAPI_OK)
+		{
+			bNvAPIInitialized = true;
+
+			NvPhysicalGpuHandle GPUHandles[NVAPI_MAX_PHYSICAL_GPUS];
+			NvU32 NumGPUs = 0;
+			if (NvAPI_EnumPhysicalGPUs(GPUHandles, &NumGPUs) == NVAPI_OK)
+			{
+				for (NvU32 GPU = 0; GPU < NumGPUs; ++GPU)
+				{
+					const NvPhysicalGpuHandle& GPUHandle = GPUHandles[GPU];
+
+					NvU32 Unused = 0, SubsystemId = 0, RevisionId = 0, ExtDeviceId = 0;
+					if (NvAPI_GPU_GetPCIIdentifiers(GPUHandle, &Unused, &SubsystemId, &RevisionId, &ExtDeviceId) != NVAPI_OK)
+					{
+						continue;
+					}
+
+					if (ExtDeviceId != Adapter.DXGIAdapterDesc.DeviceId || SubsystemId != Adapter.DXGIAdapterDesc.SubSysId || RevisionId != Adapter.DXGIAdapterDesc.Revision)
+					{
+						continue;
+					}
+
+					NV_GPU_ARCH_INFO ArchInfo;
+					ArchInfo.version = NV_GPU_ARCH_INFO_VER_2;
+					if (NvAPI_GPU_GetArchInfo(GPUHandle, &ArchInfo) != NVAPI_OK)
+					{
+						continue;
+					}
+
+					const uint32 Architecture = ArchInfo.architecture;
+					UE_LOG(LogD3D11RHI, Log, TEXT("NVIDIA Architecture=%u"), Architecture);
+					GRHIGlobals.GpuInfo.AdapterGeneration = Architecture;
+				}
+			}
+		}
+		else
+		{
+			UE_LOG(LogD3D11RHI, Warning, TEXT("Failed to initialize NVAPI"));
+		}
+	}
+#endif // NVAPI_INTERFACE
+}
+
 void FD3D11DynamicRHI::InitD3DDevice()
 {
 	check( IsInGameThread() );
@@ -1157,45 +1256,7 @@ void FD3D11DynamicRHI::InitD3DDevice()
 
 #if NV_AFTERMATH
 		UE::RHICore::Nvidia::Aftermath::InitializeBeforeDeviceCreation();
-#endif 
-
-#ifdef AMD_AGS_API
-		if (IsRHIDeviceAMD() && UE::RHICore::AllowVendorDevice())
-		{
-			check(AmdAgsContext == nullptr);
-
-			// agsInit should be called before D3D device creation
-			if (agsInitialize(AGS_MAKE_VERSION(AMD_AGS_VERSION_MAJOR, AMD_AGS_VERSION_MINOR, AMD_AGS_VERSION_PATCH), nullptr, &AmdAgsContext, &AmdInfo.AmdGpuInfo) == AGS_SUCCESS)
-			{
-				AmdInfo.AmdAgsContext = AmdAgsContext;
-				bool bFoundMatchingDevice = false;
-				// Search the device list for a matching vendor ID and device ID marked as GCN
-				for (int32 DeviceIndex = 0; DeviceIndex < AmdInfo.AmdGpuInfo.numDevices; DeviceIndex++)
-				{
-					const AGSDeviceInfo& DeviceInfo = AmdInfo.AmdGpuInfo.devices[DeviceIndex];
-					GRHIDeviceIsAMDPreGCNArchitecture |= (Adapter.DXGIAdapterDesc.VendorId == DeviceInfo.vendorId) && (Adapter.DXGIAdapterDesc.DeviceId == DeviceInfo.deviceId) && (DeviceInfo.asicFamily == AGSDeviceInfo::AsicFamily_PreGCN);
-					bFoundMatchingDevice |= (Adapter.DXGIAdapterDesc.VendorId == DeviceInfo.vendorId) && (Adapter.DXGIAdapterDesc.DeviceId == DeviceInfo.deviceId);
-				}
-				check(bFoundMatchingDevice);
-
-				if (GRHIDeviceIsAMDPreGCNArchitecture)
-				{
-					UE_LOG(LogD3D11RHI, Log, TEXT("AMD Pre GCN architecture detected, some driver workarounds will be in place"));
-				}
-			}
-			else
-			{
-				FMemory::Memzero(&AmdInfo, sizeof(AmdInfo));
-				// If agsInit returns anything but AGS_SUCCESS, the context pointer should be
-				// guaranteed to be NULL, but we'll set it here explicitly, just to be safe.
-				AmdAgsContext = nullptr;
-			}
-		}
-		else
-		{
-			FMemory::Memzero(&AmdInfo, sizeof(AmdInfo));
-		}
-#endif //AMD_AGS_API
+#endif
 
 		D3D_FEATURE_LEVEL ActualFeatureLevel = (D3D_FEATURE_LEVEL)0;
 
@@ -1423,22 +1484,11 @@ void FD3D11DynamicRHI::InitD3DDevice()
 		GRHISupportsMultithreadedResources = GRHISupportsAsyncTextureCreation;
 
 #ifdef NVAPI_INTERFACE
-
-		if (IsRHIDeviceNVIDIA() && UE::RHICore::AllowVendorDevice())
+		if (IsRHIDeviceNVIDIA() && UE::RHICore::AllowVendorDevice() && bNvAPIInitialized)
 		{
-			NvAPI_Status NvStatus;
-			NvStatus = NvAPI_Initialize();
-			if (NvStatus == NVAPI_OK)
-			{
-				NvStatus = NvAPI_D3D11_IsNvShaderExtnOpCodeSupported(Direct3DDevice, NV_EXTN_OP_UINT64_ATOMIC, &GRHISupportsAtomicUInt64);
-				if (NvStatus != NVAPI_OK)
-				{
-					UE_LOG(LogD3D11RHI, Warning, TEXT("Failed to query support for 64 bit atomics"));
-				}
-			}
-			else
+			if (NvAPI_D3D11_IsNvShaderExtnOpCodeSupported(Direct3DDevice, NV_EXTN_OP_UINT64_ATOMIC, &GRHISupportsAtomicUInt64) != NVAPI_OK)
 			{
-				UE_LOG(LogD3D11RHI, Warning, TEXT("Failed to initialize NVAPI"));
+				UE_LOG(LogD3D11RHI, Warning, TEXT("Failed to query support for 64 bit atomics"));
 			}
 		}
 
