Use static assert instead of relying on the compiler to report a missing return value
#rnx

#Synced-CL 51135664

diff --git a/Engine/Source/Developer/TraceAnalysis/Public/Trace/TraceWriter.h b/Engine/Source/Developer/TraceAnalysis/Public/Trace/TraceWriter.h
--- a/Engine/Source/Developer/TraceAnalysis/Public/Trace/TraceWriter.h
+++ b/Engine/Source/Developer/TraceAnalysis/Public/Trace/TraceWriter.h
@@ -207,9 +207,8 @@ struct FTraceWriterFieldTypeResolver
 
 		else if constexpr (std::is_same_v<std::remove_const_t<NativeType>, AnsiString>) return EType::AnsiString;
 		else if constexpr (std::is_same_v<std::remove_const_t<NativeType>, WideString>) return EType::WideString;
-
-		// No default return value (to generate a compile error for unsupported native types)!
-	}
+		else static_assert(sizeof(NativeType) == 0, "Unknown NativeType");
+	} //-V591
 };
 
 ////////////////////////////////////////////////////////////////////////////////////////////////////
diff --git a/Engine/Source/Runtime/D3D12RHI/Private/D3D12BarriersFactory.h b/Engine/Source/Runtime/D3D12RHI/Private/D3D12BarriersFactory.h
--- a/Engine/Source/Runtime/D3D12RHI/Private/D3D12BarriersFactory.h
+++ b/Engine/Source/Runtime/D3D12RHI/Private/D3D12BarriersFactory.h
@@ -198,10 +198,14 @@ private:
 			{
 				return new TFirst::AdapterImplType();
 			}
-			if constexpr (InComponentType == EComponentType::Context)
+			else if constexpr (InComponentType == EComponentType::Context)
 			{
 				return new TFirst::ContextImplType();
 			}
+			else
+			{
+				static_assert(sizeof(InComponentType) == 0, "Unknown adapter type");
+			}
 		}
 		else if (TFirst::ImplType == PreferredType)
 		{
@@ -209,16 +213,20 @@ private:
 			{
 				return new TFirst::AdapterImplType();
 			}
-			if constexpr (InComponentType == EComponentType::Context)
+			else if constexpr (InComponentType == EComponentType::Context)
 			{
 				return new TFirst::ContextImplType();
 			}
+			else
+			{
+				static_assert(sizeof(InComponentType) == 0, "Unknown adapter type");
+			}
 		}
 		else
 		{
 			return CreateBarriersInternal<InComponentType, TRest...>(PreferredType);
 		}
-	}
+	} //-V591
 
 	template <typename T, typename TFirst, typename... TRest>
 	static consteval bool CheckThatLastTypeInPackIs()
