[HWRT / D3D12] Fix stale entries in ReferencedGeometries when rebuilding AS without recreating FRHIRayTracingScene.
- also removed unused mutex.

#Synced-CL 51150711

diff --git a/Engine/Source/Runtime/D3D12RHI/Private/D3D12RayTracing.cpp b/Engine/Source/Runtime/D3D12RHI/Private/D3D12RayTracing.cpp
--- a/Engine/Source/Runtime/D3D12RHI/Private/D3D12RayTracing.cpp
+++ b/Engine/Source/Runtime/D3D12RHI/Private/D3D12RayTracing.cpp
@@ -4829,7 +4829,7 @@ void FD3D12CommandContext::RHIBuildAccelerationStructures(TConstArrayView<FRayTr
 		FD3D12Buffer* ScratchBuffer = RetrieveObject<FD3D12Buffer>(SceneBuildParams.ScratchBuffer);
 		FD3D12Buffer* InstanceBuffer = RetrieveObject<FD3D12Buffer>(SceneBuildParams.InstanceBuffer);
 
-		Scene->ReferencedGeometries.Reserve(SceneBuildParams.ReferencedGeometries.Num());
+		Scene->ReferencedGeometries.Empty(SceneBuildParams.ReferencedGeometries.Num());
 
 		for (FRHIRayTracingGeometry* ReferencedGeometry : SceneBuildParams.ReferencedGeometries)
 		{
diff --git a/Engine/Source/Runtime/D3D12RHI/Private/D3D12RayTracing.h b/Engine/Source/Runtime/D3D12RHI/Private/D3D12RayTracing.h
--- a/Engine/Source/Runtime/D3D12RHI/Private/D3D12RayTracing.h
+++ b/Engine/Source/Runtime/D3D12RHI/Private/D3D12RayTracing.h
@@ -269,9 +269,6 @@ public:
 	void UpdateResidency(FD3D12CommandContext& CommandContext) const;
 
 	bool bBuilt = false;
-
-private:
-	UE::FMutex Mutex;
 };
 
 // Manages all the pending BLAS compaction requests
