PR #14366: Adding support for Linear Swept Spheres (LSS) for ray tracing of  hair strands

Adds support for LSS primitive in D3D12 and Vulkan RHIs via NV extensions.

Adds support for LSS ray tracing for hair strands (when supported by the target hardware).

Original code submitted by [REDACTED_PERSON]

#Synced-CL 51205428

diff --git a/Engine/Config/VulkanPC/DataDrivenPlatformInfo.ini b/Engine/Config/VulkanPC/DataDrivenPlatformInfo.ini
--- a/Engine/Config/VulkanPC/DataDrivenPlatformInfo.ini
+++ b/Engine/Config/VulkanPC/DataDrivenPlatformInfo.ini
@@ -177,6 +177,7 @@ bSupportsRayTracingShaders = true
 bSupportsPathTracing = true
 bSupportsRayTracingCallableShaders = true
 bSupportsRayTracingProceduralPrimitive = true
+bSupportsRayTracingLinearSweptSpheres = true
 StaticShaderBindingLayoutSupport=RayTracingOnly
 
 bSupportsGPUScene = true
diff --git a/Engine/Config/Windows/DataDrivenPlatformInfo.ini b/Engine/Config/Windows/DataDrivenPlatformInfo.ini
--- a/Engine/Config/Windows/DataDrivenPlatformInfo.ini
+++ b/Engine/Config/Windows/DataDrivenPlatformInfo.ini
@@ -88,6 +88,7 @@ bSupportsInlineRayTracing=true
 bInlineRayTracingRequiresBindless=true
 bSupportsRayTracingCallableShaders=true
 bSupportsRayTracingProceduralPrimitive=true
+bSupportsRayTracingLinearSweptSpheres=true
 bSupportsRayTracingIndirectInstanceData=true
 bSupportsRayTracingClusterOps=true
 bSupportsPathTracing=true
diff --git a/Engine/Plugins/Runtime/HairStrands/Source/HairStrandsCore/Private/GroomAsset.cpp b/Engine/Plugins/Runtime/HairStrands/Source/HairStrandsCore/Private/GroomAsset.cpp
--- a/Engine/Plugins/Runtime/HairStrands/Source/HairStrandsCore/Private/GroomAsset.cpp
+++ b/Engine/Plugins/Runtime/HairStrands/Source/HairStrandsCore/Private/GroomAsset.cpp
@@ -935,6 +935,8 @@ bool UGroomAsset::HasGeometryType(EGroomGeometryType Type) const
 	return false;
 }
 
+bool GetSupportHairStrandsLinearSweptSpheres(EShaderPlatform InShaderPlatform);
+
 TArray<FHairVertexFactoryTypesPerMaterialData> UGroomAsset::CollectVertexFactoryTypesPerMaterialData(EShaderPlatform ShaderPlatform)
 {
 	TArray<FHairVertexFactoryTypesPerMaterialData> VFsPerMaterials;
@@ -964,7 +966,15 @@ TArray<FHairVertexFactoryTypesPerMaterialData> UGroomAsset::CollectVertexFactory
 		// Strands
 		{
 			const int32 MaterialIndex = GetMaterialIndex(GetHairGroupsRendering()[GroupIt].MaterialSlotName);
-			AddVFPerMaterial(MaterialIndex, EHairGeometryType::Strands, &FHairStrandsVertexFactory::StaticType);
+
+			if (GetSupportHairStrandsLinearSweptSpheres(ShaderPlatform) && IsHairStrandsLinearSweptSpheresEnabled())
+			{
+				AddVFPerMaterial(MaterialIndex, EHairGeometryType::Strands, &FLinearSweptSpheresHairStrandsVertexFactory::StaticType);
+			}
+			else
+			{
+				AddVFPerMaterial(MaterialIndex, EHairGeometryType::Strands, &FHairStrandsVertexFactory::StaticType);
+			}
 		}
 
 		// Cards
diff --git a/Engine/Plugins/Runtime/HairStrands/Source/HairStrandsCore/Private/GroomComponent.cpp b/Engine/Plugins/Runtime/HairStrands/Source/HairStrandsCore/Private/GroomComponent.cpp
--- a/Engine/Plugins/Runtime/HairStrands/Source/HairStrandsCore/Private/GroomComponent.cpp
+++ b/Engine/Plugins/Runtime/HairStrands/Source/HairStrandsCore/Private/GroomComponent.cpp
@@ -87,6 +87,7 @@ static FAutoConsoleVariableRef CVarHairStrands_ForceVelocityOutput(TEXT("r.HairS
 // Forward declarations
 
 bool IsHairManualSkinCacheEnabled();
+bool GetSupportHairStrandsLinearSweptSpheres(EShaderPlatform InShaderPlatform);
 
 ////////////////////////////////////////////////////////////////////////////////////////////////////////////////////////
 
@@ -130,6 +131,7 @@ FHairGroupInstance::~FHairGroupInstance()
 #endif
 		Strands.DebugCurveAttributeBuffer.Release();
 
+		InternalResourceRelease(Strands.LinearSweptSphereVertexFactory);
 		InternalResourceRelease(Strands.VertexFactory);
 	}
 
@@ -486,9 +488,22 @@ public:
 			HairInstance->bForceCards = Component->bUseCards;
 			HairInstance->bUpdatePositionOffset = Component->RegisteredMeshComponent != nullptr;
 
-			if (HairInstance->Strands.IsValid() && HairInstance->Strands.VertexFactory == nullptr)
+			if (HairInstance->Strands.IsValid())
 			{
-				HairInstance->Strands.VertexFactory = new FHairStrandsVertexFactory(HairInstance, FeatureLevel, "FStrandsHairSceneProxy");
+				if (GetSupportHairStrandsLinearSweptSpheres(ShaderPlatform) && IsHairStrandsLinearSweptSpheresEnabled())
+				{
+					if (HairInstance->Strands.LinearSweptSphereVertexFactory == nullptr)
+					{
+						HairInstance->Strands.LinearSweptSphereVertexFactory = new FLinearSweptSpheresHairStrandsVertexFactory(HairInstance, FeatureLevel, "FLinearSweptSphereStrandsHairSceneProxy");
+					}
+				}
+				else 
+				{
+					if (HairInstance->Strands.VertexFactory == nullptr)
+					{
+						HairInstance->Strands.VertexFactory = new FHairStrandsVertexFactory(HairInstance, FeatureLevel, "FStrandsHairSceneProxy");
+					}
+				}
 			}
 
 			// (Experimental) For now only use the resource of the first group
@@ -988,8 +1003,9 @@ public:
 		}
 		else // if (GeometryType == EHairGeometryType::Strands)
 		{
-			VertexFactory = (FVertexFactory*)Instance->Strands.VertexFactory;
-			PrimitiveIdMode = Instance->Strands.VertexFactory->GetPrimitiveIdMode(FeatureLevel);
+			check(Instance->Strands.LinearSweptSphereVertexFactory != nullptr || Instance->Strands.VertexFactory != nullptr);
+			VertexFactory = Instance->Strands.LinearSweptSphereVertexFactory != nullptr ? static_cast<const FVertexFactory*>(Instance->Strands.LinearSweptSphereVertexFactory) : static_cast<const FVertexFactory*>(Instance->Strands.VertexFactory);
+			PrimitiveIdMode = Instance->Strands.LinearSweptSphereVertexFactory != nullptr ? Instance->Strands.LinearSweptSphereVertexFactory->GetPrimitiveIdMode(FeatureLevel) : Instance->Strands.VertexFactory->GetPrimitiveIdMode(FeatureLevel);
 			HairVertexCount = Instance->HairGroupPublicData->GetActiveStrandsPointCount();
 			MaxVertexIndex = HairVertexCount * GetPointToVertexCount();
 			bUseCulling = Instance->Strands.bCullingEnable;
diff --git a/Engine/Plugins/Runtime/HairStrands/Source/HairStrandsCore/Private/GroomManager.cpp b/Engine/Plugins/Runtime/HairStrands/Source/HairStrandsCore/Private/GroomManager.cpp
--- a/Engine/Plugins/Runtime/HairStrands/Source/HairStrandsCore/Private/GroomManager.cpp
+++ b/Engine/Plugins/Runtime/HairStrands/Source/HairStrandsCore/Private/GroomManager.cpp
@@ -85,6 +85,7 @@ uint32 GetHairRaytracingProceduralSplits();
 void CreateHairStrandsDebugAttributeBuffer(FRDGBuilder& GraphBuilder, FRDGExternalBuffer* DebugAttributeBuffer, uint32 VertexCount, const FName& OwnerName);
 FHairGroupPublicData::FVertexFactoryInput InternalComputeHairStrandsVertexInputData(FRDGBuilder* GraphBuilder, const FHairGroupInstance* Instance, EGroomViewMode ViewMode);
 EGroomCacheType GetHairInstanceCacheType(const FHairGroupInstance* Instance);
+bool GetSupportHairStrandsLinearSweptSpheres(EShaderPlatform InShaderPlatform);
 
 void AddHairClusterAABBPass(
 	FRDGBuilder& GraphBuilder,
@@ -1815,14 +1816,13 @@ static void RunHairStrandsInterpolation_Strands(
 				InstanceData.bNeedRaytracingBuild = !InstanceData.Instance->Strands.RenRaytracingResource->bIsRTGeometryInitialized;
 				if (InstanceData.bNeedRaytracingBuild || InstanceData.bNeedRaytracingUpdate)
 				{
-					const bool bProceduralPrimitive = InstanceData.Instance->Strands.RenRaytracingResource->bProceduralPrimitive;
 					AddGenerateRaytracingGeometryPass(
 						GraphBuilder,
 						ShaderMap,
 						ShaderPrintData,
 						InstanceData.Instance->RegisteredIndex,
 						InstanceData.ActivePointCount,
-						bProceduralPrimitive,
+						InstanceData.Instance->Strands.RenRaytracingResource->RayTracingPrimitiveType,
 						ProceduralSplits,
 						HairRadiusRT,
 						HairRootScaleRT,
@@ -2865,7 +2865,16 @@ static bool SelectValidLOD(
 #if RHI_RAYTRACING
 		if (Instance->Strands.RenRaytracingResource)	{ Instance->Strands.RenRaytracingResource->Allocate	(GraphBuilder, EHairResourceLoadingType::Async, ResourceStatus, RequestedCurveCount, RequestedPointCount, Instance->Strands.RenRaytracingResourceOwned/*bAllowDeallocation*/); }
 #endif
-		Instance->Strands.VertexFactory->InitResources(GraphBuilder.RHICmdList);
+		
+		if (Instance->Strands.LinearSweptSphereVertexFactory != nullptr)
+		{
+			Instance->Strands.LinearSweptSphereVertexFactory->InitResources(GraphBuilder.RHICmdList);
+		}
+		else
+		{
+			check(Instance->Strands.VertexFactory != nullptr);
+			Instance->Strands.VertexFactory->InitResources(GraphBuilder.RHICmdList);
+		}
 
 		// Early initialization, so that when filtering MeshBatch block in SceneVisiblity, we can use this value to know 
 		// if the hair instance is visible or not (i.e., HairLengthScale > 0)
diff --git a/Engine/Plugins/Runtime/HairStrands/Source/HairStrandsCore/Private/GroomResources.cpp b/Engine/Plugins/Runtime/HairStrands/Source/HairStrandsCore/Private/GroomResources.cpp
--- a/Engine/Plugins/Runtime/HairStrands/Source/HairStrandsCore/Private/GroomResources.cpp
+++ b/Engine/Plugins/Runtime/HairStrands/Source/HairStrandsCore/Private/GroomResources.cpp
@@ -10,6 +10,7 @@
 #include "RenderingThread.h"
 #include "Engine/AssetUserData.h"
 #include "HairStrandsVertexFactory.h"
+#include "HairStrandsInterface.h"
 #include "Misc/Paths.h"
 #include "RHIStaticStates.h"
 #include "RHIUtilities.h"
@@ -1981,16 +1982,25 @@ void FHairCardsInterpolationResource::InternalRelease()
 
 #if RHI_RAYTRACING
 bool GetSupportHairStrandsProceduralPrimitive(EShaderPlatform InShaderPlatform);
+bool GetSupportHairStrandsLinearSweptSpheres(EShaderPlatform InShaderPlatform);
 
-static void ConvertPointToRaytracingVertexCount(bool bProceduralPrimitive, uint32 InPointCount, uint32& OutVertexCount, uint32& OutIndexCount)
+static void ConvertPointToRaytracingVertexCount(EHairStrandsRayTracingPrimitiveType RayTracingPrimitiveType, uint32 InPointCount, uint32& OutVertexCount, uint32& OutIndexCount)
 {
-	if (bProceduralPrimitive)
+	if (RayTracingPrimitiveType == EHairStrandsRayTracingPrimitiveType::LinearSweptSphere)
+	{
+		// 1 vertex at the beginning to indicate invalid hair strand
+		OutVertexCount = InPointCount + 1;
+		OutIndexCount  = InPointCount * 2;
+	}
+	else if (RayTracingPrimitiveType == EHairStrandsRayTracingPrimitiveType::Procedural)
 	{
 		// only allocate space for primitive AABBs
 		OutVertexCount = InPointCount * 2 * STRANDS_PROCEDURAL_INTERSECTOR_MAX_SPLITS;
+		OutIndexCount  = 0;
 	}
 	else
 	{
+		// triangle fallback path
 		OutVertexCount = InPointCount * 4;
 		OutIndexCount  = InPointCount * 8 * 3;
 	}
@@ -2011,7 +2021,16 @@ FHairStrandsRaytracingResource::FHairStrandsRaytracingResource(const FHairStrand
 	MaxVertexCount = 0;
 	MaxIndexCount = 0;
 	MaxAvailableCurveCount = 0;
-	bProceduralPrimitive = GetSupportHairStrandsProceduralPrimitive(GMaxRHIShaderPlatform);
+
+	RayTracingPrimitiveType = EHairStrandsRayTracingPrimitiveType::Triangle;
+	if (GetSupportHairStrandsLinearSweptSpheres(GMaxRHIShaderPlatform) && IsHairStrandsLinearSweptSpheresEnabled())
+	{
+		RayTracingPrimitiveType = EHairStrandsRayTracingPrimitiveType::LinearSweptSphere;
+	}
+	else if (GetSupportHairStrandsProceduralPrimitive(GMaxRHIShaderPlatform))
+	{
+		RayTracingPrimitiveType = EHairStrandsRayTracingPrimitiveType::Procedural;
+	}
 }
 
 FHairStrandsRaytracingResource::FHairStrandsRaytracingResource(const FHairCardsBulkData& InData, const FHairResourceName& ResourceName, const FName& InOwnerName) :
@@ -2028,7 +2047,7 @@ void FHairStrandsRaytracingResource::InternalAllocate(FRDGBuilder& GraphBuilder,
 {
 	if (bOwnBuffers)
 	{
-		ConvertPointToRaytracingVertexCount(bProceduralPrimitive, InPointCount, MaxVertexCount, MaxIndexCount);
+		ConvertPointToRaytracingVertexCount(RayTracingPrimitiveType, InPointCount, MaxVertexCount, MaxIndexCount);
 
 		InternalCreateVertexBufferRDG<FHairStrandsRaytracingFormat>(GraphBuilder, MaxVertexCount, PositionBuffer, ToHairResourceDebugName(TEXT("Hair.StrandsRaytracing_PositionBuffer"), ResourceName), OwnerName, EHairResourceUsageType::Dynamic);
 		InternalCreateStructuredBufferRDG<FHairStrandsIndexFormat>(GraphBuilder, MaxIndexCount, IndexBuffer, ToHairResourceDebugName(TEXT("Hair.StrandsRaytracing_IndexBuffer"), ResourceName), OwnerName, EHairResourceUsageType::Dynamic);
diff --git a/Engine/Plugins/Runtime/HairStrands/Source/HairStrandsCore/Private/HairStrandsInterpolation.cpp b/Engine/Plugins/Runtime/HairStrands/Source/HairStrandsCore/Private/HairStrandsInterpolation.cpp
--- a/Engine/Plugins/Runtime/HairStrands/Source/HairStrandsCore/Private/HairStrandsInterpolation.cpp
+++ b/Engine/Plugins/Runtime/HairStrands/Source/HairStrandsCore/Private/HairStrandsInterpolation.cpp
@@ -78,6 +78,9 @@ bool GetHairStrandsBindingDataValidation()
 	#endif
 }
 
+bool GetSupportHairStrandsProceduralPrimitive(EShaderPlatform InShaderPlatform);
+bool GetSupportHairStrandsLinearSweptSpheres(EShaderPlatform InShaderPlatform);
+
 uint32 GetHairStrandsForceSingleGuideInterpolation()
 {
 	return GStrandHairInterpolationForceSingleGuide > 0;
@@ -1246,8 +1249,8 @@ class FHairRaytracingGeometryCS : public FGlobalShader
 	SHADER_USE_PARAMETER_STRUCT(FHairRaytracingGeometryCS, FGlobalShader);
 
 	class FCulling : SHADER_PERMUTATION_BOOL("PERMUTATION_CULLING");
-	class FProceduralPrimitive : SHADER_PERMUTATION_BOOL("PERMUTATION_PROCEDURAL_PRIMITIVE"); 
-	using FPermutationDomain = TShaderPermutationDomain<FCulling, FProceduralPrimitive>;
+	class FPrimitiveType : SHADER_PERMUTATION_ENUM_CLASS("PRIMITIVE_TYPE", EHairStrandsRayTracingPrimitiveType);
+	using FPermutationDomain = TShaderPermutationDomain<FCulling, FPrimitiveType>;
 
 	BEGIN_SHADER_PARAMETER_STRUCT(FParameters, )
 		SHADER_PARAMETER(uint32, PointCount)
@@ -1272,7 +1275,29 @@ class FHairRaytracingGeometryCS : public FGlobalShader
 
 public:
 	static uint32 GetGroupSize() { return HAIR_VERTEXCOUNT_GROUP_SIZE; }
-	static bool ShouldCompilePermutation(const FGlobalShaderPermutationParameters& Parameters) { return IsHairStrandsSupported(EHairStrandsShaderType::Strands, Parameters.Platform); }
+	static bool ShouldCompilePermutation(const FGlobalShaderPermutationParameters& Parameters) 
+	{ 
+		FPermutationDomain PermutationVector(Parameters.PermutationId);
+
+		const EHairStrandsRayTracingPrimitiveType PrimitiveType = PermutationVector.Get<FPrimitiveType>();
+
+		if ((PrimitiveType == EHairStrandsRayTracingPrimitiveType::LinearSweptSphere) && !GetSupportHairStrandsLinearSweptSpheres(Parameters.Platform))
+		{
+			return false;
+		}
+
+		if ((PrimitiveType == EHairStrandsRayTracingPrimitiveType::Procedural) && !GetSupportHairStrandsProceduralPrimitive(Parameters.Platform))
+		{
+			return false;
+		}
+
+		if ((PrimitiveType == EHairStrandsRayTracingPrimitiveType::Triangle) && GetSupportHairStrandsProceduralPrimitive(Parameters.Platform))
+		{
+			return false;
+		}
+
+		return IsHairStrandsSupported(EHairStrandsShaderType::Strands, Parameters.Platform); 
+	}
 	static void ModifyCompilationEnvironment(const FGlobalShaderPermutationParameters& Parameters, FShaderCompilerEnvironment& OutEnvironment)
 	{
 		FGlobalShader::ModifyCompilationEnvironment(Parameters, OutEnvironment);
@@ -1298,7 +1323,7 @@ void AddGenerateRaytracingGeometryPass(
 	const FShaderPrintData* ShaderPrintData,
 	uint32 InstanceRegisteredIndex,
 	uint32 PointCount,
-	bool bProceduralPrimitive,
+	EHairStrandsRayTracingPrimitiveType RayTracingPrimitiveType,
 	int ProceduralSplits,
 	float HairRadius,
 	float RootScale,
@@ -1339,13 +1364,13 @@ void AddGenerateRaytracingGeometryPass(
 
 	FHairRaytracingGeometryCS::FPermutationDomain PermutationVector;
 	PermutationVector.Set<FHairRaytracingGeometryCS::FCulling>(bCullingEnable);
-	PermutationVector.Set<FHairRaytracingGeometryCS::FProceduralPrimitive>(bProceduralPrimitive);
+	PermutationVector.Set<FHairRaytracingGeometryCS::FPrimitiveType>(RayTracingPrimitiveType);
 	TShaderMapRef<FHairRaytracingGeometryCS> ComputeShader(ShaderMap, PermutationVector);
 
 	const FIntVector DispatchCount(FMath::DivideAndRoundUp(PointCount, FHairRaytracingGeometryCS::GetGroupSize()), 1, 1);
 	FComputeShaderUtils::AddPass(
 		GraphBuilder,
-		RDG_EVENT_NAME("HairStrands::GenerateRaytracingGeometry(bCulling:%d, Procedural:%d)", bCullingEnable, bProceduralPrimitive),
+		RDG_EVENT_NAME("HairStrands::GenerateRaytracingGeometry(bCulling:%d, PrimitiveType:%d)", bCullingEnable, (int32)RayTracingPrimitiveType),
 		ComputeShader,
 		Parameters,
 		DispatchCount);
@@ -1420,7 +1445,7 @@ static void BuildHairAccelerationStructure_Strands(
 	const FString& AssetDebugName,
 	const FString& MeshComponentName,
 	uint32 LODIndex,
-	bool bProceduralPrimitive,
+	EHairStrandsRayTracingPrimitiveType RayTracingPrimitiveType,
 	int ProceduralSplits)
 {
 	FRayTracingGeometryInitializer Initializer;
@@ -1436,7 +1461,15 @@ static void BuildHairAccelerationStructure_Strands(
 	const FName OwnerName(FString::Printf(TEXT("%s [LOD%d]"), *MeshComponentName, LODIndex));
 	Initializer.OwnerName = OwnerName;
 #endif
-	if (bProceduralPrimitive)
+	
+	if (RayTracingPrimitiveType == EHairStrandsRayTracingPrimitiveType::LinearSweptSphere)
+	{
+		Initializer.IndexBuffer = IndexBuffer;
+		Initializer.IndexBufferOffset = 0;
+		Initializer.GeometryType = RTGT_LinearSweptSpheres;
+		Initializer.TotalPrimitiveCount = RayTracingIndexCount / 2;
+	}
+	else if (RayTracingPrimitiveType == EHairStrandsRayTracingPrimitiveType::Procedural)
 	{
 		Initializer.GeometryType = RTGT_Procedural;
 		Initializer.TotalPrimitiveCount = (RaytracingVertexCount / (2 * STRANDS_PROCEDURAL_INTERSECTOR_MAX_SPLITS)) * ProceduralSplits;
@@ -1454,7 +1487,14 @@ static void BuildHairAccelerationStructure_Strands(
 	FRayTracingGeometrySegment Segment;
 	Segment.VertexBuffer = PositionBuffer;
 	Segment.VertexBufferElementType = FHairStrandsRaytracingFormat::VertexElementType;
-	if (bProceduralPrimitive)
+
+	if (RayTracingPrimitiveType == EHairStrandsRayTracingPrimitiveType::LinearSweptSphere)
+	{
+		Segment.VertexBufferStride = FHairStrandsRaytracingFormat::SizeInByte;
+		Segment.MaxVertices = RaytracingVertexCount;
+		Segment.NumPrimitives = Initializer.TotalPrimitiveCount;
+	}
+	else if (RayTracingPrimitiveType == EHairStrandsRayTracingPrimitiveType::Procedural)
 	{
 		// stride covers one AABB (which is made up of two vertices)
 		Segment.VertexBufferStride = FHairStrandsRaytracingFormat::SizeInByte * 2;
@@ -1569,7 +1609,7 @@ void AddBuildStrandsAccelerationStructurePass(
 	FRDGBufferRef Raytracing_IndexBuffer)
 {
 	const uint32 HairLODIndex = Instance->HairGroupPublicData->GetIntLODIndex();
-	const bool bProceduralPrimitive = Instance->Strands.RenRaytracingResource->bProceduralPrimitive;
+	const EHairStrandsRayTracingPrimitiveType RayTracingPrimitiveType = Instance->Strands.RenRaytracingResource->RayTracingPrimitiveType;
 
 	FStrandsResourceBLASParameters* Parameters = GraphBuilder.AllocParameters < FStrandsResourceBLASParameters > ();
 	Parameters->PositionBuffer = Raytracing_PositionBuffer;
@@ -1579,7 +1619,7 @@ void AddBuildStrandsAccelerationStructurePass(
 		RDG_EVENT_NAME("HairStrands::UpdateBLAS(Strands)"),
 		Parameters,
 		ERDGPassFlags::NeverCull | ERDGPassFlags::Compute,
-		[Instance, HairLODIndex, bNeedUpdate, bProceduralPrimitive, ProceduralSplits](FRHICommandListImmediate& RHICmdList)
+		[Instance, HairLODIndex, bNeedUpdate, RayTracingPrimitiveType, ProceduralSplits](FRHICommandListImmediate& RHICmdList)
 		{
 			SCOPED_GPU_MASK(RHICmdList, FRHIGPUMask::All());
 	
@@ -1589,7 +1629,13 @@ void AddBuildStrandsAccelerationStructurePass(
 				FBufferRHIRef PositionBuffer(Instance->Strands.RenRaytracingResource->PositionBuffer.Buffer->GetRHI());
 	
 				// no index buffer needed for procedural primitive
-				FBufferRHIRef IndexBuffer(bProceduralPrimitive ? nullptr : Instance->Strands.RenRaytracingResource->IndexBuffer.Buffer->GetRHI());
+				FBufferRHIRef IndexBuffer = nullptr;
+				if (RayTracingPrimitiveType == EHairStrandsRayTracingPrimitiveType::Triangle ||
+				    RayTracingPrimitiveType == EHairStrandsRayTracingPrimitiveType::LinearSweptSphere)
+				{
+					IndexBuffer = Instance->Strands.RenRaytracingResource->IndexBuffer.Buffer->GetRHI();
+				}
+
 				Instance->Strands.CachedProceduralSplits = ProceduralSplits;
 				BuildHairAccelerationStructure_Strands(
 					RHICmdList,
@@ -1601,7 +1647,7 @@ void AddBuildStrandsAccelerationStructurePass(
 					Instance->Debug.GroomAssetName,
 					Instance->Debug.MeshComponentName,
 					HairLODIndex,
-					bProceduralPrimitive,
+					RayTracingPrimitiveType,
 					ProceduralSplits
 				);
 			}
diff --git a/Engine/Plugins/Runtime/HairStrands/Source/HairStrandsCore/Private/HairStrandsVertexFactory.cpp b/Engine/Plugins/Runtime/HairStrands/Source/HairStrandsCore/Private/HairStrandsVertexFactory.cpp
--- a/Engine/Plugins/Runtime/HairStrands/Source/HairStrandsCore/Private/HairStrandsVertexFactory.cpp
+++ b/Engine/Plugins/Runtime/HairStrands/Source/HairStrandsCore/Private/HairStrandsVertexFactory.cpp
@@ -19,12 +19,18 @@
 
 #define VF_STRANDS_SUPPORT_GPU_SCENE 1
 #define VF_STRANDS_PROCEDURAL_INTERSECTOR 1
+#define VF_STRANDS_LINEAR_SWEPT_SPHERES_INTERSECTOR 1
 
 bool GetSupportHairStrandsProceduralPrimitive(EShaderPlatform InShaderPlatform)
 {
 	return VF_STRANDS_PROCEDURAL_INTERSECTOR && FDataDrivenShaderPlatformInfo::GetSupportsRayTracingProceduralPrimitive(InShaderPlatform);
 }
 
+bool GetSupportHairStrandsLinearSweptSpheres(EShaderPlatform InShaderPlatform)
+{
+	return VF_STRANDS_LINEAR_SWEPT_SPHERES_INTERSECTOR && FDataDrivenShaderPlatformInfo::GetSupportsRayTracingLinearSweptSpheres(InShaderPlatform);
+}
+
 /////////////////////////////////////////////////////////////////////////////////////////
 
 template<typename T> inline void VFS_BindParam(FMeshDrawSingleShaderBindings& ShaderBindings, const FShaderResourceParameter& Param, T* Value)	{ if (Param.IsBound() && Value) ShaderBindings.Add(Param, Value); }
@@ -152,6 +158,7 @@ FHairStrandsVertexFactoryUniformShaderParameters FHairGroupInstance::GetHairStan
 
 IMPLEMENT_UNIFORM_BUFFER_STRUCT(FHairStrandsVertexFactoryUniformShaderParameters, "HairStrandsVF");
 
+template<bool bSupportLinearSweptSpheres>
 class HAIRSTRANDSCORE_API FHairStrandsVertexFactoryShaderParameters : public FVertexFactoryShaderParameters
 {
 	DECLARE_TYPE_LAYOUT(FHairStrandsVertexFactoryShaderParameters, NonVirtual);
@@ -173,22 +180,33 @@ public:
 		FVertexInputStreamArray& VertexStreams
 	) const
 	{
-		const FHairStrandsVertexFactory* VF = static_cast<const FHairStrandsVertexFactory*>(VertexFactory);
+		const THairStrandsVertexFactory<bSupportLinearSweptSpheres>* VF = static_cast<const THairStrandsVertexFactory<bSupportLinearSweptSpheres>*>(VertexFactory);
 		ShaderBindings.Add(Shader->GetUniformBufferParameter<FHairStrandsVertexFactoryUniformShaderParameters>(), VF->Data.Instance->Strands.UniformBuffer);
 	}
 };
 
-IMPLEMENT_TYPE_LAYOUT(FHairStrandsVertexFactoryShaderParameters);
+IMPLEMENT_TEMPLATE_TYPE_LAYOUT(template<>, FHairStrandsVertexFactoryShaderParameters<true>);
+IMPLEMENT_TEMPLATE_TYPE_LAYOUT(template<>, FHairStrandsVertexFactoryShaderParameters<false>);
 
 /**
  * Should we cache the material's shadertype on this platform with this vertex factory? 
  */
-bool FHairStrandsVertexFactory::ShouldCompilePermutation(const FVertexFactoryShaderPermutationParameters& Parameters)
+template<bool bSupportLinearSweptSpheres>
+bool THairStrandsVertexFactory<bSupportLinearSweptSpheres>::ShouldCompilePermutation(const FVertexFactoryShaderPermutationParameters& Parameters)
 {
+	if constexpr (bSupportLinearSweptSpheres)
+	{
+		if (!GetSupportHairStrandsLinearSweptSpheres(Parameters.Platform))
+		{
+			return false;
+		}
+	}
+
 	return (Parameters.MaterialParameters.MaterialDomain == MD_Surface && Parameters.MaterialParameters.bIsUsedWithHairStrands && IsHairStrandsSupported(EHairStrandsShaderType::Strands, Parameters.Platform)) || Parameters.MaterialParameters.bIsSpecialEngineMaterial;
 }
 
-void FHairStrandsVertexFactory::ModifyCompilationEnvironment(const FVertexFactoryShaderPermutationParameters& Parameters, FShaderCompilerEnvironment& OutEnvironment)
+template<bool bSupportLinearSweptSpheres>
+void THairStrandsVertexFactory<bSupportLinearSweptSpheres>::ModifyCompilationEnvironment(const FVertexFactoryShaderPermutationParameters& Parameters, FShaderCompilerEnvironment& OutEnvironment)
 {
 	bool bUseGPUSceneAndPrimitiveIdStream = false;
 #if VF_STRANDS_SUPPORT_GPU_SCENE
@@ -198,13 +216,30 @@ void FHairStrandsVertexFactory::ModifyCompilationEnvironment(const FVertexFactor
 		// TODO: support GPUScene on mobile
 		&& !IsMobilePlatform(Parameters.Platform);
 #endif
-	const bool bUseProceduralIntersection = GetSupportHairStrandsProceduralPrimitive(Parameters.Platform);
 	OutEnvironment.SetDefine(TEXT("VF_SUPPORTS_PRIMITIVE_SCENE_DATA"), bUseGPUSceneAndPrimitiveIdStream);
 	OutEnvironment.SetDefine(TEXT("HAIR_STRAND_MESH_FACTORY"), 1);
-	OutEnvironment.SetDefine(TEXT("ENABLE_PROCEDURAL_INTERSECTOR"), bUseProceduralIntersection);
+
+	if constexpr (bSupportLinearSweptSpheres)
+	{
+		OutEnvironment.SetDefine(TEXT("ENABLE_HAIRSTRANDS_LINEAR_SWEPT_SPHERE_INTERSECTOR"), 1);
+		OutEnvironment.SetDefine(TEXT("ALLOW_CUSTOM_RAYTRACING_HIT_TYPE"), 1);
+	}
+	else
+	{
+		const bool bUseProceduralIntersection = GetSupportHairStrandsProceduralPrimitive(Parameters.Platform);
+		OutEnvironment.SetDefine(TEXT("ENABLE_HAIRSTRANDS_PROCEDURAL_INTERSECTOR"), bUseProceduralIntersection);
+
+		// This is not strictly needed since it always return front face hit in IS.
+		// However, hard coding front face check may allow compiler to optimize code more aggressively.
+		if (bUseProceduralIntersection)
+		{
+			OutEnvironment.SetDefine(TEXT("ALLOW_CUSTOM_RAYTRACING_HIT_TYPE"), 1);
+		}
+	}
 }
 
-void FHairStrandsVertexFactory::ValidateCompiledResult(const FVertexFactoryType* Type, EShaderPlatform Platform, const FShaderParameterMap& ParameterMap, TArray<FString>& OutErrors)
+template<bool bSupportLinearSweptSpheres>
+void THairStrandsVertexFactory<bSupportLinearSweptSpheres>::ValidateCompiledResult(const FVertexFactoryType* Type, EShaderPlatform Platform, const FShaderParameterMap& ParameterMap, TArray<FString>& OutErrors)
 {
 #if VF_STRANDS_SUPPORT_GPU_SCENE
 	if (Type->SupportsPrimitiveIdStream()
@@ -217,7 +252,8 @@ void FHairStrandsVertexFactory::ValidateCompiledResult(const FVertexFactoryType*
 #endif
 }
 
-void FHairStrandsVertexFactory::SetData(const FDataType& InData)
+template<bool bSupportLinearSweptSpheres>
+void THairStrandsVertexFactory<bSupportLinearSweptSpheres>::SetData(const FDataType& InData)
 {
 	Data = InData;
 	UpdateRHI(FRHICommandListImmediate::Get());
@@ -227,9 +263,10 @@ void FHairStrandsVertexFactory::SetData(const FDataType& InData)
 * Copy the data from another vertex factory
 * @param Other - factory to copy from
 */
-void FHairStrandsVertexFactory::Copy(const FHairStrandsVertexFactory& Other)
+template<bool bSupportLinearSweptSpheres>
+void THairStrandsVertexFactory<bSupportLinearSweptSpheres>::Copy(const THairStrandsVertexFactory& Other)
 {
-	FHairStrandsVertexFactory* VertexFactory = this;
+	THairStrandsVertexFactory* VertexFactory = this;
 	const FDataType* DataCopy = &Other.Data;
 	ENQUEUE_RENDER_COMMAND(FHairStrandsVertexFactoryCopyData)(/*UE::RenderCommandPipe::Groom,*/
 		[VertexFactory, DataCopy](FRHICommandListImmediate& RHICmdList)
@@ -239,7 +276,8 @@ void FHairStrandsVertexFactory::Copy(const FHairStrandsVertexFactory& Other)
 	BeginUpdateResourceRHI(this);
 }
 
-void FHairStrandsVertexFactory::InitResources(FRHICommandListBase& RHICmdList)
+template<bool bSupportLinearSweptSpheres>
+void THairStrandsVertexFactory<bSupportLinearSweptSpheres>::InitResources(FRHICommandListBase& RHICmdList)
 {
 	if (bIsInitialized)
 		return;
@@ -272,36 +310,53 @@ void FHairStrandsVertexFactory::InitResources(FRHICommandListBase& RHICmdList)
 	Data.Instance->Strands.UniformBuffer = FHairStrandsUniformBuffer::CreateUniformBufferImmediate(Parameters, UniformBuffer_MultiFrame);
 }
 
-void FHairStrandsVertexFactory::GetPSOPrecacheVertexFetchElements(EVertexInputStreamType VertexInputStreamType, FVertexDeclarationElementList& Elements)
+template<bool bSupportLinearSweptSpheres>
+void THairStrandsVertexFactory<bSupportLinearSweptSpheres>::GetPSOPrecacheVertexFetchElements(EVertexInputStreamType VertexInputStreamType, FVertexDeclarationElementList& Elements)
 {
 #if VF_STRANDS_SUPPORT_GPU_SCENE
 	Elements.Add(FVertexElement(0, 0, VET_UInt, 13, sizeof(uint32), true));
 #endif
 }
 
-EPrimitiveIdMode FHairStrandsVertexFactory::GetPrimitiveIdMode(ERHIFeatureLevel::Type In) const
+template<bool bSupportLinearSweptSpheres>
+EPrimitiveIdMode THairStrandsVertexFactory<bSupportLinearSweptSpheres>::GetPrimitiveIdMode(ERHIFeatureLevel::Type In) const
 {
 	return PrimID_ForceZero;
 }
 
-IMPLEMENT_VERTEX_FACTORY_PARAMETER_TYPE(FHairStrandsVertexFactory, SF_Vertex,		FHairStrandsVertexFactoryShaderParameters);
-IMPLEMENT_VERTEX_FACTORY_PARAMETER_TYPE(FHairStrandsVertexFactory, SF_Pixel,		FHairStrandsVertexFactoryShaderParameters);
-IMPLEMENT_VERTEX_FACTORY_PARAMETER_TYPE(FHairStrandsVertexFactory, SF_Compute,		FHairStrandsVertexFactoryShaderParameters);
+IMPLEMENT_VERTEX_FACTORY_PARAMETER_TYPE(FHairStrandsVertexFactory, SF_Vertex,      FHairStrandsVertexFactoryShaderParameters<false>);
+IMPLEMENT_VERTEX_FACTORY_PARAMETER_TYPE(FHairStrandsVertexFactory, SF_Pixel,       FHairStrandsVertexFactoryShaderParameters<false>);
+IMPLEMENT_VERTEX_FACTORY_PARAMETER_TYPE(FHairStrandsVertexFactory, SF_Compute,     FHairStrandsVertexFactoryShaderParameters<false>);
 #if RHI_RAYTRACING
-IMPLEMENT_VERTEX_FACTORY_PARAMETER_TYPE(FHairStrandsVertexFactory, SF_RayHitGroup,	FHairStrandsVertexFactoryShaderParameters);
+IMPLEMENT_VERTEX_FACTORY_PARAMETER_TYPE(FHairStrandsVertexFactory, SF_RayHitGroup, FHairStrandsVertexFactoryShaderParameters<false>);
 #endif
 
-void FHairStrandsVertexFactory::InitRHI(FRHICommandListBase& RHICmdList)
+#if VF_STRANDS_LINEAR_SWEPT_SPHERES_INTERSECTOR
+IMPLEMENT_VERTEX_FACTORY_PARAMETER_TYPE(FLinearSweptSpheresHairStrandsVertexFactory, SF_Vertex,     FHairStrandsVertexFactoryShaderParameters<true>);
+IMPLEMENT_VERTEX_FACTORY_PARAMETER_TYPE(FLinearSweptSpheresHairStrandsVertexFactory, SF_Pixel,      FHairStrandsVertexFactoryShaderParameters<true>);
+IMPLEMENT_VERTEX_FACTORY_PARAMETER_TYPE(FLinearSweptSpheresHairStrandsVertexFactory, SF_Compute,    FHairStrandsVertexFactoryShaderParameters<true>);
+#if RHI_RAYTRACING
+IMPLEMENT_VERTEX_FACTORY_PARAMETER_TYPE(FLinearSweptSpheresHairStrandsVertexFactory, SF_RayHitGroup, FHairStrandsVertexFactoryShaderParameters<true>);
+#endif
+#endif
+
+template<bool bSupportLinearSweptSpheres>
+void THairStrandsVertexFactory<bSupportLinearSweptSpheres>::InitRHI(FRHICommandListBase& RHICmdList)
 {
 	// Nothing as the initialize runs only on first use
 }
 
-void FHairStrandsVertexFactory::ReleaseRHI()
+template<bool bSupportLinearSweptSpheres>
+void THairStrandsVertexFactory<bSupportLinearSweptSpheres>::ReleaseRHI()
 {
 	FVertexFactory::ReleaseRHI();
 }
 
-IMPLEMENT_VERTEX_FACTORY_TYPE(FHairStrandsVertexFactory, "/Engine/Private/HairStrands/HairStrandsVertexFactory.ush",
+// Explicit instantiations
+template class THairStrandsVertexFactory<true>;
+template class THairStrandsVertexFactory<false>;
+
+IMPLEMENT_TEMPLATE_VERTEX_FACTORY_TYPE(template<>, FHairStrandsVertexFactory, "/Engine/Private/HairStrands/HairStrandsVertexFactory.ush",
 	EVertexFactoryFlags::UsedWithMaterials
 	| EVertexFactoryFlags::SupportsDynamicLighting
 	| EVertexFactoryFlags::SupportsPrecisePrevWorldPos
@@ -312,3 +367,15 @@ IMPLEMENT_VERTEX_FACTORY_TYPE(FHairStrandsVertexFactory, "/Engine/Private/HairSt
 	| EVertexFactoryFlags::SupportsManualVertexFetch
 	| EVertexFactoryFlags::SupportsPSOPrecaching
 );
+
+IMPLEMENT_TEMPLATE_VERTEX_FACTORY_TYPE(template<>, FLinearSweptSpheresHairStrandsVertexFactory, "/Engine/Private/HairStrands/HairStrandsVertexFactory.ush",
+	EVertexFactoryFlags::UsedWithMaterials
+	| EVertexFactoryFlags::SupportsDynamicLighting
+	| EVertexFactoryFlags::SupportsPrecisePrevWorldPos
+	| EVertexFactoryFlags::SupportsCachingMeshDrawCommands
+	| (VF_STRANDS_SUPPORT_GPU_SCENE ? EVertexFactoryFlags::SupportsPrimitiveIdStream : EVertexFactoryFlags::None)
+	| EVertexFactoryFlags::SupportsRayTracing
+	| EVertexFactoryFlags::SupportsManualVertexFetch
+	| EVertexFactoryFlags::SupportsPSOPrecaching
+);
+
diff --git a/Engine/Plugins/Runtime/HairStrands/Source/HairStrandsCore/Public/GroomInstance.h b/Engine/Plugins/Runtime/HairStrands/Source/HairStrandsCore/Public/GroomInstance.h
--- a/Engine/Plugins/Runtime/HairStrands/Source/HairStrandsCore/Public/GroomInstance.h
+++ b/Engine/Plugins/Runtime/HairStrands/Source/HairStrandsCore/Public/GroomInstance.h
@@ -17,7 +17,6 @@
 class UMeshComponent;
 class UGroomComponent;
 class FHairCardsVertexFactory;
-class FHairStrandsVertexFactory;
 enum class EGroomCacheType : uint8;
 enum class EGroomViewMode : uint8;
 
@@ -151,7 +150,8 @@ struct FHairGroupInstance : public FHairStrandsInstance
 		FHairGroupInstanceModifer Modifier;
 
 		FHairStrandsUniformBuffer UniformBuffer;
-		FHairStrandsVertexFactory* VertexFactory = nullptr;
+		FHairStrandsVertexFactory*                   VertexFactory = nullptr;
+		FLinearSweptSpheresHairStrandsVertexFactory* LinearSweptSphereVertexFactory = nullptr;
 
 		EHairInterpolationType HairInterpolationType = EHairInterpolationType::NoneSkinning;		
 		bool bCullingEnable = false; // Indicates if culling is enabled for this hair strands data.
diff --git a/Engine/Plugins/Runtime/HairStrands/Source/HairStrandsCore/Public/GroomResources.h b/Engine/Plugins/Runtime/HairStrands/Source/HairStrandsCore/Public/GroomResources.h
--- a/Engine/Plugins/Runtime/HairStrands/Source/HairStrandsCore/Public/GroomResources.h
+++ b/Engine/Plugins/Runtime/HairStrands/Source/HairStrandsCore/Public/GroomResources.h
@@ -689,7 +689,7 @@ struct FHairStrandsRaytracingResource : public FHairCommonResource
 	uint32 MaxIndexCount = 0;
 	bool bOwnBuffers = false;
 	bool bIsRTGeometryInitialized = false;
-	bool bProceduralPrimitive = false;
+	EHairStrandsRayTracingPrimitiveType RayTracingPrimitiveType = EHairStrandsRayTracingPrimitiveType::Triangle;
 };
 #endif
 
diff --git a/Engine/Plugins/Runtime/HairStrands/Source/HairStrandsCore/Public/HairStrandsInterpolation.h b/Engine/Plugins/Runtime/HairStrands/Source/HairStrandsCore/Public/HairStrandsInterpolation.h
--- a/Engine/Plugins/Runtime/HairStrands/Source/HairStrandsCore/Public/HairStrandsInterpolation.h
+++ b/Engine/Plugins/Runtime/HairStrands/Source/HairStrandsCore/Public/HairStrandsInterpolation.h
@@ -196,7 +196,7 @@ void AddGenerateRaytracingGeometryPass(
 	const FShaderPrintData* ShaderPrintData,
 	uint32 InstanceRegisteredIndex,
 	uint32 PointCount,
-	bool bProceduralPrimitive,
+	EHairStrandsRayTracingPrimitiveType RayTracingPrimitiveType,
 	int ProceduralSplits,
 	float HairRadius,
 	float RootScale,
diff --git a/Engine/Plugins/Runtime/HairStrands/Source/HairStrandsCore/Public/HairStrandsVertexFactory.h b/Engine/Plugins/Runtime/HairStrands/Source/HairStrandsCore/Public/HairStrandsVertexFactory.h
--- a/Engine/Plugins/Runtime/HairStrands/Source/HairStrandsCore/Public/HairStrandsVertexFactory.h
+++ b/Engine/Plugins/Runtime/HairStrands/Source/HairStrandsCore/Public/HairStrandsVertexFactory.h
@@ -20,10 +20,13 @@ struct FHairGroupInstance;
 
 /**
  * A vertex factory which simply transforms explicit vertex attributes from local to world space.
+ * 
+ * This is templated to allow different behavior for Linear Swept Spheres, which do not require an intersection shader.
  */
-class FHairStrandsVertexFactory : public FVertexFactory
+template<bool bSupportLinearSweptSpheres>
+class THairStrandsVertexFactory : public FVertexFactory
 {
-	DECLARE_VERTEX_FACTORY_TYPE(FHairStrandsVertexFactory);
+	DECLARE_VERTEX_FACTORY_TYPE(THairStrandsVertexFactory);
 public:
 	
 	struct FDataType
@@ -31,7 +34,7 @@ public:
 		FHairGroupInstance* Instance = nullptr;
 	};
 
-	FHairStrandsVertexFactory(FHairGroupInstance* Instance, ERHIFeatureLevel::Type InFeatureLevel, const char* InDebugName)
+	THairStrandsVertexFactory(FHairGroupInstance* Instance, ERHIFeatureLevel::Type InFeatureLevel, const char* InDebugName)
 		: FVertexFactory(InFeatureLevel)
 		, DebugName(InDebugName)
 	{		
@@ -58,7 +61,7 @@ public:
 	* Copy the data from another vertex factory
 	* @param Other - factory to copy from
 	*/
-	UE_API void Copy(const FHairStrandsVertexFactory& Other);
+	UE_API void Copy(const THairStrandsVertexFactory& Other);
 
 	// FRenderResource interface.
 	UE_API virtual void InitRHI(FRHICommandListBase& RHICmdList) override;
@@ -84,4 +87,7 @@ protected:
 	} DebugName;
 };
 
+using FLinearSweptSpheresHairStrandsVertexFactory = THairStrandsVertexFactory<true>;
+using FHairStrandsVertexFactory = THairStrandsVertexFactory<false>;
+
 #undef UE_API
diff --git a/Engine/Shaders/Private/HairStrands/HairStrandsRaytracingGeometry.usf b/Engine/Shaders/Private/HairStrands/HairStrandsRaytracingGeometry.usf
--- a/Engine/Shaders/Private/HairStrands/HairStrandsRaytracingGeometry.usf
+++ b/Engine/Shaders/Private/HairStrands/HairStrandsRaytracingGeometry.usf
@@ -8,6 +8,11 @@
 
 #if SHADER_RT_GEOMETRY
 
+// This has to match EHairStrandsRayTracingPrimitiveType enum in HairStrandsInterface.h
+#define PRIMITIVE_TYPE_PROCEDURAL          0
+#define PRIMITIVE_TYPE_TRIANGLE            1
+#define PRIMITIVE_TYPE_LINEAR_SWEPT_SPHERE 2
+
 uint PointCount;
 
 float HairStrandsVF_HairRadius;
@@ -19,7 +24,7 @@ StructuredBuffer<float4> PositionOffsetBuffer;
 ByteAddressBuffer PositionBuffer;
 Buffer<float4> 	 TangentBuffer;
 RWBuffer<float4> OutputPositionBuffer;	
-#if PERMUTATION_PROCEDURAL_PRIMITIVE 
+#if PRIMITIVE_TYPE == PRIMITIVE_TYPE_PROCEDURAL
 uint             RaytracingProceduralSplits;
 #else
 RWStructuredBuffer<uint> OutputIndexBuffer;
@@ -58,7 +63,7 @@ float3x3 GetTangentBasis(uint ControlPointId)
 	return float3x3(Tx, Ty, Tz);
 }
 
-#if PERMUTATION_PROCEDURAL_PRIMITIVE
+#if PRIMITIVE_TYPE == PRIMITIVE_TYPE_PROCEDURAL
 void SetAABB(uint InVertexIndex, float3 Lo, float3 Hi)
 {
 	OutputPositionBuffer[2 * InVertexIndex + 0] = float4(Lo.x, Lo.y, Lo.z, Hi.x);
@@ -69,11 +74,15 @@ void SetAABB(uint InVertexIndex, float3 Lo, float3 Hi)
 // a whole segment of invalid values
 void SetInvalidSegment(uint InVertexIndex)
 {
-#if PERMUTATION_PROCEDURAL_PRIMITIVE
+#if PRIMITIVE_TYPE == PRIMITIVE_TYPE_PROCEDURAL
 	const float NaN = asfloat(0x7FC00000);
 	for (uint i = 0; i < RaytracingProceduralSplits; i++)
 		SetAABB(RaytracingProceduralSplits * InVertexIndex + i, NaN, NaN);
-#else
+#elif PRIMITIVE_TYPE == PRIMITIVE_TYPE_LINEAR_SWEPT_SPHERE
+	const uint BasePoint = InVertexIndex * 2;
+	OutputIndexBuffer[BasePoint    ] = 0;
+	OutputIndexBuffer[BasePoint + 1] = 0;
+#else // PRIMITIVE_TYPE == PRIMITIVE_TYPE_TRIANGLE
 	const uint BasePoint = InVertexIndex * 4;
 	const uint TrisPerSegment = 8;
 	const uint IdxsPerSegment = TrisPerSegment * 3;
@@ -113,7 +122,7 @@ void SetCapSegment(uint InVertexIndex)
 #endif
 }
 
-#if !PERMUTATION_PROCEDURAL_PRIMITIVE
+#if PRIMITIVE_TYPE == PRIMITIVE_TYPE_TRIANGLE
 void SetBodySegment(uint InVertexIndex)
 {
 	const uint BotPoint = InVertexIndex * 4;
@@ -151,12 +160,35 @@ float GetAABBHalfAreaForSegment(float4 A, float4 B)
 
 bool IsValidPosition(FHairControlPoint P)
 {
-	return all(!isinf(P.Position)) && all(!isnan(P.Position)) && all(P.Position <= INFINITE_FLOAT);
+	return all(abs(P.Position) < INFINITE_FLOAT);
 }
 
 [numthreads(GROUP_SIZE, 1, 1)]
 void MainCS(uint DispatchThreadId : SV_DispatchThreadID)
 {
+#if PERMUTATION_CULLING == 1
+	const uint MaxPointCount = HairStrandsVF_bCullingEnable ? HairStrandsVF_CullingIndirectBuffer[3] : PointCount;
+#else
+	const uint MaxPointCount = PointCount;
+#endif
+
+#if PRIMITIVE_TYPE == PRIMITIVE_TYPE_LINEAR_SWEPT_SPHERE
+	// preserve the first entry for Nan to indicate invalid hair strand
+	const uint VertexBufferIndexShift = 1;
+	
+	if (DispatchThreadId < VertexBufferIndexShift)
+	{
+		const float NaN = asfloat(0x7FC00000);
+		OutputPositionBuffer[DispatchThreadId] = NaN;
+	}
+
+	if (DispatchThreadId < MaxPointCount)
+	{
+		const FHairControlPoint CP = GetControlPoint(DispatchThreadId);
+		OutputPositionBuffer[DispatchThreadId + VertexBufferIndexShift] = float4(CP.Position, CP.WorldRadius);
+	}
+#endif
+
 	uint VertexIndex0 = DispatchThreadId;
 	uint VertexIndex1 = VertexIndex0 + 1;
 
@@ -172,8 +204,7 @@ void MainCS(uint DispatchThreadId : SV_DispatchThreadID)
 #if PERMUTATION_CULLING == 1
 	if (HairStrandsVF_bCullingEnable)
 	{
-		const uint VertexCountAfterCulling = HairStrandsVF_CullingIndirectBuffer[3];
-		if (VertexIndex1 >= VertexCountAfterCulling)
+		if (VertexIndex1 >= MaxPointCount)
 		{
 			SetInvalidSegment(VertexIndex0);
 			return;
@@ -185,7 +216,7 @@ void MainCS(uint DispatchThreadId : SV_DispatchThreadID)
 #endif
 
 	FHairControlPoint P0 = GetControlPoint(SourceIndex0);
-#if PERMUTATION_PROCEDURAL_PRIMITIVE 
+#if PRIMITIVE_TYPE == PRIMITIVE_TYPE_PROCEDURAL 
 	if (P0.Type == HAIR_CONTROLPOINT_END || !IsValidPosition(P0))
 	{
 		// this vertex doesn't map to a valid segment
@@ -197,7 +228,7 @@ void MainCS(uint DispatchThreadId : SV_DispatchThreadID)
 		FHairControlPoint P1 = GetControlPoint(SourceIndex1);
 		if (!IsValidPosition(P1))
 		{
-			SetInvalidSegment(VertexIndex1);
+			SetInvalidSegment(VertexIndex0);
 		}
 		else
 		{
@@ -242,7 +273,38 @@ void MainCS(uint DispatchThreadId : SV_DispatchThreadID)
 			}
 		}
 	}
-#else
+#elif PRIMITIVE_TYPE == PRIMITIVE_TYPE_LINEAR_SWEPT_SPHERE
+	if (P0.Type == HAIR_CONTROLPOINT_END || !IsValidPosition(P0))
+	{
+		// this vertex doesn't map to a valid segment
+		SetInvalidSegment(VertexIndex0);
+	}
+	else
+	{
+		// fetch the other point and create the aabb around the line
+		FHairControlPoint P1 = GetControlPoint(SourceIndex1);
+		if (!IsValidPosition(P1) || VertexIndex1 >= MaxPointCount)
+		{
+			SetInvalidSegment(VertexIndex0);
+		}
+		else
+		{
+			float4 A = float4(P0.Position, P0.WorldRadius);
+			float4 B = float4(P1.Position, P1.WorldRadius);
+			if ((A.w == 0 && B.w == 0) || length2(B.xyz - A.xyz) == 0 || any(!IsFinite(A)) || any(!IsFinite(B)))
+			{
+				// this segment is degenerate, don't bother adding it
+				SetInvalidSegment(VertexIndex0);
+			}
+			else
+			{
+				const uint BaseOutIndex = VertexIndex0 * 2;
+				OutputIndexBuffer[BaseOutIndex    ] = SourceIndex0 + VertexBufferIndexShift;
+				OutputIndexBuffer[BaseOutIndex + 1] = SourceIndex1 + VertexBufferIndexShift;
+			}
+		}
+	}
+#else // PRIMITIVE_TYPE == PRIMITIVE_TYPE_TRIANGLE
 	// expand current vertex into 4 points around the curve
 	const float3x3 TangentBasis = GetTangentBasis(SourceIndex0);
 	const float3 N = TangentBasis[0];
@@ -268,8 +330,19 @@ void MainCS(uint DispatchThreadId : SV_DispatchThreadID)
 	}
 	else
 	{
-		// regular point, build the body of the cylinder
-		SetBodySegment(VertexIndex0);
+		FHairControlPoint P1 = GetControlPoint(SourceIndex1);
+		// Skip segment if second vertex is degenerate, or we have reached the very last vertex
+		// TODO: is there an off-by-one error elsewhere? the last point seems to always be near the origin for some reason (but only in this mode)
+		// TODO: Check if this issue impacts the other modes as well
+		if (!IsValidPosition(P1) || SourceIndex1 == (MaxPointCount - 1))
+		{
+			SetInvalidSegment(VertexIndex0);
+		}
+		else
+		{
+			// regular point, build the body of the cylinder
+			SetBodySegment(VertexIndex0);
+		}
 	}
 #endif
 }
diff --git a/Engine/Shaders/Private/HairStrands/HairStrandsVertexFactory.ush b/Engine/Shaders/Private/HairStrands/HairStrandsVertexFactory.ush
--- a/Engine/Shaders/Private/HairStrands/HairStrandsVertexFactory.ush
+++ b/Engine/Shaders/Private/HairStrands/HairStrandsVertexFactory.ush
@@ -17,6 +17,8 @@
 	#error Hair triangle geometry type needs to be defined
 #endif
 
+#define ENABLE_HAIRSTRANDS_TRIANGLE_INTERSECTOR !(ENABLE_HAIRSTRANDS_LINEAR_SWEPT_SPHERE_INTERSECTOR || ENABLE_HAIRSTRANDS_PROCEDURAL_INTERSECTOR)
+
 #if MANUAL_VERTEX_FETCH
 	Buffer<float4> VertexFetch_InstanceOriginBuffer;
 	Buffer<float4> VertexFetch_InstanceTransformBuffer;
@@ -33,7 +35,7 @@
 // HAIR_STRAND_MESH_FACTORY
 // Define used in certain shader like DeepShadow for running specific code (coverage computation for instance) 
 // when vertices are produced by the strand hair vertex factory 
-// This is set by compilation enviromenent of the vertex factory
+// This is set by compilation environment of the vertex factory
 ///////////////////////////////////////////////////////////////////////////////
 
 
@@ -86,8 +88,11 @@ struct FVertexFactoryInput
 {
 #if RAYHITGROUPSHADER
 	float4 Position;
-	uint TriangleId;          // from RayHit
-	uint TriangleVertexId;    // 0,1,2
+	uint PrimitiveId;          // from RayHit (either TriangleId or SegmendId)
+
+#if ENABLE_HAIRSTRANDS_TRIANGLE_INTERSECTOR
+	uint PrimitiveVertexId;    // 0,1,2 (Vertex within the triangle)
+#endif
 #endif
 
 	// Dynamic instancing related attributes with InstanceIdOffset : ATTRIBUTE13
@@ -116,10 +121,23 @@ FVertexInfo GetVertexInfo(FVertexFactoryInput Input)
 
 	FVertexInfo VertexInfo;
 #if RAYHITGROUPSHADER
+
+#if ENABLE_HAIRSTRANDS_LINEAR_SWEPT_SPHERE_INTERSECTOR
+	VertexInfo.IsTip		= 0;
+	uint BaseIndex			= VertexId;
+#elif ENABLE_HAIRSTRANDS_PROCEDURAL_INTERSECTOR
+	VertexInfo.IsTip		= 0;
+	uint BaseIndex			= VertexId;
+#else
+	// Triangle intersection
 	uint BaseIndex			= VertexId / 4;
 	// TODO: is there any easy way we could handle end-caps?
-	VertexInfo.IsTip		= (Input.TriangleId % 2) == 0 ? Input.TriangleVertexId != 2 : Input.TriangleVertexId == 0;
-#elif USE_HAIR_TRIANGLE_STRIP
+	VertexInfo.IsTip		= (Input.PrimitiveId % 2) == 0 ? Input.PrimitiveVertexId != 2 : Input.PrimitiveVertexId == 0;
+#endif
+
+#else
+
+#if USE_HAIR_TRIANGLE_STRIP
 	uint QuadIndex 			= VertexId % 2;
 	uint BaseIndex			= VertexId / 2;
 	VertexInfo.IsTip		= 0;
@@ -129,6 +147,8 @@ FVertexInfo GetVertexInfo(FVertexFactoryInput Input)
 	uint BaseIndex			= VertexId / 6;
 	VertexInfo.IsTip		= QuadIndex == 0 || QuadIndex == 2 || QuadIndex == 4 ? 0 : 1;
 	VertexInfo.IsLeft		= QuadIndex == 0 || QuadIndex == 1 || QuadIndex == 5 ? 1 : 0;
+#endif
+
 #endif
 	VertexInfo.HairControlPointId = BaseIndex;
 	VertexInfo.VertexIndex = BaseIndex + VertexInfo.IsTip;
@@ -171,18 +191,41 @@ half3x3 AssembleTangentToWorldForHairStrand(half3 TangentToWorld0, half4 Tangent
 }
 
 #if RAYHITGROUPSHADER
-FVertexFactoryInput LoadVertexFactoryInputForHGS(uint TriangleIndex, int VertexIndex)
+FHairControlPoint GetVertexPosition(const FVertexInfo VertexInfo, bool bInvalidJointVertex);
+
+FVertexFactoryInput LoadVertexFactoryInputForHGS(uint PrimitiveIndex, int VertexIndex)
 {
 	FVertexFactoryInput Input = (FVertexFactoryInput)0;
 
-#if !ENABLE_PROCEDURAL_INTERSECTOR
-	FTriangleBaseAttributes TriangleBase = LoadTriangleBaseAttributes(TriangleIndex);
+#if ENABLE_HAIRSTRANDS_LINEAR_SWEPT_SPHERE_INTERSECTOR
+	if (VertexIndex == 0 || VertexIndex == 1)
+	{
+		Input.VertexId = PrimitiveIndex + VertexIndex;
+		Input.PrimitiveId = PrimitiveIndex;
+	}
+#elif ENABLE_HAIRSTRANDS_PROCEDURAL_INTERSECTOR
+	if (VertexIndex == 0 || VertexIndex == 1)
+	{
+		const uint BasePrimitiveIndex = PrimitiveIndex / HairStrandsVF.RaytracingProceduralSplits;
 
-	Input.Position = float4(TriangleBase.LocalPositions[VertexIndex], 1.0f);
-	Input.VertexId = TriangleBase.Indices[VertexIndex];
-	Input.TriangleId = TriangleIndex;
-	Input.TriangleVertexId = VertexIndex;
-#endif // !ENABLE_PROCEDURAL_INTERSECTOR
+		Input.VertexId = BasePrimitiveIndex + VertexIndex;
+		Input.PrimitiveId = PrimitiveIndex;
+	}
+#else
+	{
+		FTriangleBaseAttributes TriangleBase = LoadTriangleBaseAttributes(PrimitiveIndex);
+
+		// The data in this buffer has 32 bits per float, while the data read from previous buffer only has 16 bits.
+		// This mismatch will cause motion vector to be non-zero when there is no motion.
+		// Input.Position = float4(TriangleBase.LocalPositions[VertexIndex], 1.0f);
+		Input.VertexId = TriangleBase.Indices[VertexIndex];
+		Input.PrimitiveId = PrimitiveIndex;
+		Input.PrimitiveVertexId = VertexIndex;
+	}
+#endif // !ENABLE_HAIRSTRANDS_PROCEDURAL_INTERSECTOR && !ENABLE_HAIRSTRANDS_LINEAR_SWEPT_SPHERE_INTERSECTOR
+	
+	const FVertexInfo VertexInfo = GetVertexInfo(Input);
+	Input.Position = float4(GetVertexPosition(VertexInfo, true).Position, 1.f);
 
 	// Note: GetInstanceUserData() stores the GPU-Scene instance ID
 	VF_GPUSCENE_SET_INPUT_FOR_RT(Input, GetInstanceUserData(), 0U);
@@ -551,111 +594,74 @@ FMaterialPixelParameters GetMaterialPixelParameters(float3 RayOrigin, float3 Ray
 
 	FSceneDataIntermediates SceneData = VF_GPUSCENE_GET_INTERMEDIATES(Input); // NOTE: Input is not used when VF_USE_PRIMITIVE_SCENE_DATA == 0
 
-	const float2 Barycentrics = HitAttributes.GetBarycentrics();
-
-#if ENABLE_PROCEDURAL_INTERSECTOR
-	uint VertexIndex0 = HitPrimitiveIndex / HairStrandsVF.RaytracingProceduralSplits;
-	uint VertexIndex1 = VertexIndex0 + 1;
+	const float EffectiveRadius = HairStrandsVF.Radius * HairStrandsVF.RaytracingRadiusScale;
 
-	if (HairStrandsVF.bCullingEnable)
-	{
-		VertexIndex0 = HairStrandsVF.CullingIndexBuffer[VertexIndex0];
-		VertexIndex1 = HairStrandsVF.CullingIndexBuffer[VertexIndex1];
-	}
+	const float2 Barycentrics = HitAttributes.GetBarycentrics();
 
-	const float EffectiveRadius = HairStrandsVF.Radius * HairStrandsVF.RaytracingRadiusScale;
 	float UCoord = Barycentrics.x;
 	float VCoord = Barycentrics.y;
-	const float3 PositionOffset = GetPositionOffset();
-	FHairControlPoint P0 = ReadHairControlPoint(
-		HairStrandsVF.PositionBuffer,
-		VertexIndex0,
-		PositionOffset,
-		EffectiveRadius,
-		HairStrandsVF.RootScale,
-		HairStrandsVF.TipScale);
-	FHairControlPoint P1 = ReadHairControlPoint(
-		HairStrandsVF.PositionBuffer,
-		VertexIndex1,
-		PositionOffset,
-		EffectiveRadius,
-		HairStrandsVF.RootScale,
-		HairStrandsVF.TipScale);
 
+#if ENABLE_HAIRSTRANDS_LINEAR_SWEPT_SPHERE_INTERSECTOR
+	uint VertexIndex0 = HitPrimitiveIndex;
+#elif ENABLE_HAIRSTRANDS_PROCEDURAL_INTERSECTOR
+	uint VertexIndex0 = HitPrimitiveIndex / HairStrandsVF.RaytracingProceduralSplits;
 #else
-	const float3 Weights = float3(
-		1 - Barycentrics.x - Barycentrics.y,
-		Barycentrics.x,
-		Barycentrics.y);
-
 	// fetch triangle
-	FTriangleBaseAttributes Tri = LoadTriangleBaseAttributes(HitPrimitiveIndex);
-	// Transform vertices into world space
-	const float3 WorldPositions[3] = {
-		CalcWorldPosition(Tri.LocalPositions[0], SceneData.Primitive.LocalToWorld).xyz,
-		CalcWorldPosition(Tri.LocalPositions[1], SceneData.Primitive.LocalToWorld).xyz,
-		CalcWorldPosition(Tri.LocalPositions[2], SceneData.Primitive.LocalToWorld).xyz
-	};
-	// Compute triangle normal
-	{
-		const float3 PA = WorldPositions[1] - WorldPositions[0];
-		const float3 PB = WorldPositions[2] - WorldPositions[0];
-		const float3 Unnormalized = cross(PB, PA);
-		const float InvWorldArea = rsqrt(dot(Unnormalized, Unnormalized));
-		WorldGeoNormal = Unnormalized * InvWorldArea;
-	}
+	const FTriangleBaseAttributes Tri = LoadTriangleBaseAttributes(HitPrimitiveIndex);
 
 	// NOTE: This must match the tesselation pattern produced by SetBodySegment() in HairStrandsRaytracingGeometry.usf
 	// TODO: is there a way we could handle end-caps elegantly?
 	uint VertexIndex0 = Tri.Indices.z / 4; // last index always belongs to the "bottom" vertex of the hair segment
-	uint VertexIndex1 = VertexIndex0 + 1;
 
+	const float3 Weights = float3(1 - Barycentrics.x - Barycentrics.y, Barycentrics.x, Barycentrics.y);
+	UCoord = dot(float3(1, float((HitPrimitiveIndex % 2) == 0), 0), Weights);
+#endif
+
+	uint VertexIndex1 = VertexIndex0 + 1;
 	if (HairStrandsVF.bCullingEnable)
 	{
 		VertexIndex0 = HairStrandsVF.CullingIndexBuffer[VertexIndex0];
 		VertexIndex1 = HairStrandsVF.CullingIndexBuffer[VertexIndex1];
 	}
-
-	float UCoord = dot(float3(1, float((HitPrimitiveIndex % 2) == 0), 0), Weights);
-	float VCoord = 0;
+	
+#if !ENABLE_HAIRSTRANDS_PROCEDURAL_INTERSECTOR
+	// Azimuth distance for procedural primitives are evaluated in intersection shader.
+	// A different approach is used here for better precision
 	{
-		// find central point of the hair strand
 		const float3 PositionOffset = GetPositionOffset();
-		const float EffectiveRadius = HairStrandsVF.Radius * HairStrandsVF.RaytracingRadiusScale;
-		FHairControlPoint P0 = ReadHairControlPoint(
+		const FHairControlPoint P0 = ReadHairControlPoint(
 			HairStrandsVF.PositionBuffer,
 			VertexIndex0,
 			PositionOffset,
 			EffectiveRadius,
 			HairStrandsVF.RootScale,
 			HairStrandsVF.TipScale);
-		FHairControlPoint P1 = ReadHairControlPoint(
+		const FHairControlPoint P1 = ReadHairControlPoint(
 			HairStrandsVF.PositionBuffer,
 			VertexIndex1,
 			PositionOffset,
 			EffectiveRadius,
 			HairStrandsVF.RootScale,
 			HairStrandsVF.TipScale);
-		float4 A = float4(P0.Position, P0.WorldRadius);
-		float4 B = float4(P1.Position, P1.WorldRadius);
-		float4 C = lerp(A, B, UCoord);
+
+		float3 A = DFDemote(CalcWorldPosition(P0.Position, SceneData.Primitive.LocalToWorld));
+		float3 B = DFDemote(CalcWorldPosition(P1.Position, SceneData.Primitive.LocalToWorld));
+
 		// approximate position along the width of the ribbon
-		float3 WorldC = DFDemote(CalcWorldPosition(C.xyz, SceneData.Primitive.LocalToWorld));
-		float3 rd = RayDirection;
-		float3 cp = WorldC - RayOrigin;
-		float t = dot(cp, rd) / dot(rd, rd);
-		float3 dp = t * RayDirection - cp;
-		float d2 = length2(dp);
-		float v = saturate(sqrt(d2 / Pow2(C.w))); // TODO: account for scale of radius
-		// figure out sign of v
-		float3 WorldU = DFDemote(TransformLocalToWorld(B.xyz - A.xyz, SceneData.Primitive.LocalToWorld));
-		VCoord = v * sign(dot(dp, cross(RayDirection, WorldU)));
-		VCoord = 0.5 * VCoord + 0.5;
-	}
+		
 
+		const float3 AxisIntersection = lerp(A, B, UCoord.x);
+		const float3 IntersectedPoint = RayOrigin + RayDirection * HitT;
+		const float3 ShadingNormal = normalize(IntersectedPoint - AxisIntersection);
+		const float3 HairStrandDirection = normalize(B - A);
 
+		const float3 Proj = normalize(RayDirection - dot(RayDirection, HairStrandDirection) * HairStrandDirection);
+		const float3 ProjPerpendicular = cross(Proj, HairStrandDirection);
+		const float h = dot(ShadingNormal, ProjPerpendicular);
+		VCoord = 0.5f * h + 0.5f;
+	}
+#endif
 
-#endif // ENABLE_PROCEDURAL_INTERSECTOR
 	Result.HairControlPointId = VertexIndex0;
 	Result.HairPrimitiveUV = float2(UCoord, VCoord);
 
@@ -674,7 +680,11 @@ FMaterialPixelParameters GetMaterialPixelParameters(float3 RayOrigin, float3 Ray
 	float3 D = RayDirection;
 	WorldGeoNormal = normalize(cross(T, normalize(cross(D, T))));
 
+#if ENABLE_HAIRSTRANDS_LINEAR_SWEPT_SPHERE_INTERSECTOR
+	Result.TwoSidedSign = 1.f;
+#else
 	Result.TwoSidedSign = sign(dot(RayDirection, WorldGeoNormal));
+#endif
 
 	Result.PrimitiveId = SceneData.PrimitiveId;
 
@@ -912,14 +922,14 @@ float3 VertexFactoryGetWorldNormal(FVertexFactoryInput Input, FVertexFactoryInte
 }
 
 ////////////////////////////////////////////////////////////////////////////////////////////////////////////
-// Regular material evalution for raytracing
+// Regular material evaluation for raytracing
 
-#if RAYHITGROUP || COMPUTESHADER
-FVertexFactoryInput LoadVertexFactoryInputForDynamicUpdate(uint TriangleIndex, int VertexIndex, uint PrimitiveId, uint DrawInstanceId)
+#if COMPUTESHADER
+FVertexFactoryInput LoadVertexFactoryInputForDynamicUpdate(uint PrimitiveIndex, int VertexIndex, uint PrimitiveId, uint DrawInstanceId)
 {
 	FVertexFactoryInput Input;
 
-	Input.VertexId = TriangleIndex * 3 + VertexIndex;
+	Input.VertexId = PrimitiveIndex * 3 + VertexIndex;
 
 	return Input;
 }
@@ -928,7 +938,7 @@ uint GetNumRayTracingDynamicMeshVerticesIndirect()
 {
 	return 0;
 }
-#endif // RAYHITGROUP || COMPUTESHADER
+#endif // COMPUTESHADER
 
 #if NEEDS_VERTEX_FACTORY_INTERPOLATION
 
@@ -1003,4 +1013,12 @@ FVertexFactoryRayTracingInterpolants VertexFactoryInterpolate(FVertexFactoryRayT
 }
 #endif // NEEDS_VERTEX_FACTORY_INTERPOLATION
 
+#if RAYHITGROUPSHADER && ALLOW_CUSTOM_RAYTRACING_HIT_TYPE
+bool IsFrontFaceHit()
+{
+	// Hardware-accelerated LSS intersection tests or procedural tests won’t hit back faces.
+	return true;
+}
+#endif
+
 #include "/Engine/Private/VertexFactoryDefaultInterface.ush"
diff --git a/Engine/Shaders/Private/PathTracing/PathTracingDebug.usf b/Engine/Shaders/Private/PathTracing/PathTracingDebug.usf
--- a/Engine/Shaders/Private/PathTracing/PathTracingDebug.usf
+++ b/Engine/Shaders/Private/PathTracing/PathTracingDebug.usf
@@ -1153,9 +1153,10 @@ FPathTracingResult PathTracingDebugCore(FRayDesc Ray, uint2 LaunchIndex)
 				case PATH_TRACER_DEBUG_VIZ_SUBSTRATE_TRANSMITTANCE_N: 	Color = HitPayload.TransmittanceN; 		break;
 				case PATH_TRACER_DEBUG_VIZ_SUBSTRATE_COVERAGE_ABOVE_N: 	Color = HitPayload.CoverageAboveAlongN; break;
 				case PATH_TRACER_DEBUG_VIZ_SUBSTRATE_F90:				Color = HitPayload.SpecularEdgeColor;   break;
-				case PATH_TRACER_DEBUG_VIZ_EYE_IRIS_MASK:      Color = (HitPayload.PathTracingShadingModelID == PATH_TRACING_SHADINGMODELID_EYE) ?             HitPayload.GetEyeIrisMask()      : 0.0; break;
-				case PATH_TRACER_DEBUG_VIZ_EYE_IRIS_NORMAL:    Color = (HitPayload.PathTracingShadingModelID == PATH_TRACING_SHADINGMODELID_EYE) ? 0.5 + 0.5 * HitPayload.GetEyeIrisNormal()    : 0.0; break;
-				case PATH_TRACER_DEBUG_VIZ_EYE_CAUSTIC_NORMAL: Color = (HitPayload.PathTracingShadingModelID == PATH_TRACING_SHADINGMODELID_EYE) ? 0.5 + 0.5 * HitPayload.GetEyeCausticNormal() : 0.0; break;
+				case PATH_TRACER_DEBUG_VIZ_EYE_IRIS_MASK:      Color = (HitPayload.PathTracingShadingModelID == PATH_TRACING_SHADINGMODELID_EYE ) ?             HitPayload.GetEyeIrisMask()      : 0.0; break;
+				case PATH_TRACER_DEBUG_VIZ_EYE_IRIS_NORMAL:    Color = (HitPayload.PathTracingShadingModelID == PATH_TRACING_SHADINGMODELID_EYE ) ? 0.5 + 0.5 * HitPayload.GetEyeIrisNormal()    : 0.0; break;
+				case PATH_TRACER_DEBUG_VIZ_EYE_CAUSTIC_NORMAL: Color = (HitPayload.PathTracingShadingModelID == PATH_TRACING_SHADINGMODELID_EYE ) ? 0.5 + 0.5 * HitPayload.GetEyeCausticNormal() : 0.0; break;
+				case PATH_TRACER_DEBUG_VIZ_HAIR_PRIMITIVE_UV:  Color = (HitPayload.PathTracingShadingModelID == PATH_TRACING_SHADINGMODELID_HAIR) ? float3(HitPayload.GetHairPrimitiveUV(), 0.0) : 0.0; break;
 				case PATH_TRACER_DEBUG_VIZ_BACKFACING_NORMAL:
 				{
 					Color = FakeShading;
diff --git a/Engine/Shaders/Private/PathTracing/PathTracingMaterialHitShader.usf b/Engine/Shaders/Private/PathTracing/PathTracingMaterialHitShader.usf
--- a/Engine/Shaders/Private/PathTracing/PathTracingMaterialHitShader.usf
+++ b/Engine/Shaders/Private/PathTracing/PathTracingMaterialHitShader.usf
@@ -224,7 +224,7 @@ RAY_TRACING_ENTRY_CLOSEST_HIT(PathTracingMaterialCHS,
 
 	FPixelMaterialInputs PixelMaterialInputs;
 
-	const bool bIsFrontFace = HitKind() == HIT_KIND_TRIANGLE_FRONT_FACE;
+	const bool bIsFrontFace = IsFrontFaceHit();
 	{
 		const float4 ScreenPosition = SvPositionToResolvedScreenPosition(SvPosition);
 		MaterialParameters.CameraVector = -WorldRayDirection();
@@ -329,7 +329,8 @@ RAY_TRACING_ENTRY_CLOSEST_HIT(PathTracingMaterialCHS,
 	Payload.PrimitiveLightingChannelMask = GetPrimitive_LightingChannelMask_FromFlags(PrimitiveFlags);
 
 	Payload.HitT = RayTCurrent();
-	if (HitKind() == HIT_KIND_TRIANGLE_FRONT_FACE)
+
+	if (IsFrontFaceHit())
 	{
 		Payload.SetFrontFace();
 	}
@@ -886,6 +887,7 @@ RAY_TRACING_ENTRY_CLOSEST_HIT(PathTracingMaterialCHS,
 #if HAIR_STRAND_MESH_FACTORY
 	Payload.WorldNormal = Payload.WorldGeoNormal;
 	Payload.WorldTangent = MaterialParameters.TangentToWorld[2];
+
 	// Repurpose WorldSmoothNormal to encode hair strand direction
 	Payload.WorldSmoothNormal = normalize(MaterialParameters.WorldNormal);
 #endif
@@ -1400,7 +1402,7 @@ RAY_TRACING_ENTRY_ANY_HIT(PathTracingMaterialAHS,
 
 	FPixelMaterialInputs PixelMaterialInputs;
 
-	const bool bIsFrontFace = HitKind() == HIT_KIND_TRIANGLE_FRONT_FACE;
+	const bool bIsFrontFace = IsFrontFaceHit();
 
 	{
 		const float4 ScreenPosition = SvPositionToResolvedScreenPosition(SvPosition);
diff --git a/Engine/Shaders/Private/RayTracing/RayTracingMaterialHitShaders.usf b/Engine/Shaders/Private/RayTracing/RayTracingMaterialHitShaders.usf
--- a/Engine/Shaders/Private/RayTracing/RayTracingMaterialHitShaders.usf
+++ b/Engine/Shaders/Private/RayTracing/RayTracingMaterialHitShaders.usf
@@ -696,7 +696,7 @@ RAY_TRACING_ENTRY_CLOSEST_HIT(MaterialCHS,
 		float4 ScreenPosition = SvPositionToResolvedScreenPosition(SvPosition);
 		float3 TranslatedWorldPosition = TranslatedWorldRayOrigin() + WorldRayDirection() * RayTCurrent();
 		
-		bool bIsFrontFace = HitKind() == HIT_KIND_TRIANGLE_FRONT_FACE;
+		const bool bIsFrontFace = IsFrontFaceHit();
 
 		// #dxr_todo: UE-72130 support world position offset 
 		// #if USE_WORLD_POSITION_EXCLUDING_SHADER_OFFSETS
@@ -733,7 +733,8 @@ RAY_TRACING_ENTRY_CLOSEST_HIT(MaterialCHS,
 	Payload.PrimitiveLightingChannelMask = GetPrimitive_LightingChannelMask_FromFlags(PrimitiveFlags);
 
 	Payload.HitT = RayTCurrent();
-	if (HitKind() == HIT_KIND_TRIANGLE_FRONT_FACE)
+
+	if (IsFrontFaceHit())
 	{
 		Payload.SetFrontFace();
 	}
@@ -867,8 +868,8 @@ RAY_TRACING_ENTRY_ANY_HIT(MaterialAHS,
 	{
 		float4 ScreenPosition = SvPositionToResolvedScreenPosition(SvPosition);
 		float3 TranslatedWorldPosition = TranslatedWorldRayOrigin() + WorldRayDirection() * RayTCurrent();
-		
-		bool bIsFrontFace = HitKind() == HIT_KIND_TRIANGLE_FRONT_FACE;
+
+		const bool bIsFrontFace = IsFrontFaceHit();
 
 		MaterialParameters.CameraVector = -WorldRayDirection();
 		
diff --git a/Engine/Shaders/Private/VertexFactoryCommon.ush b/Engine/Shaders/Private/VertexFactoryCommon.ush
--- a/Engine/Shaders/Private/VertexFactoryCommon.ush
+++ b/Engine/Shaders/Private/VertexFactoryCommon.ush
@@ -206,3 +206,10 @@ float3 OctToUnit( float2 Oct )
 	}
 	return normalize(N);
 }
+
+#if RAYHITGROUPSHADER && !ALLOW_CUSTOM_RAYTRACING_HIT_TYPE
+bool IsFrontFaceHit()
+{
+    return HitKind() == HIT_KIND_TRIANGLE_FRONT_FACE;
+}
+#endif
\ No newline at end of file
diff --git a/Engine/Shaders/Shared/PathTracingDefinitions.h b/Engine/Shaders/Shared/PathTracingDefinitions.h
--- a/Engine/Shaders/Shared/PathTracingDefinitions.h
+++ b/Engine/Shaders/Shared/PathTracingDefinitions.h
@@ -127,3 +127,4 @@
 #define PATH_TRACER_DEBUG_VIZ_EYE_CAUSTIC_NORMAL        40
 #define PATH_TRACER_DEBUG_VIZ_EYE_IRIS_NORMAL           41
 #define PATH_TRACER_DEBUG_VIZ_BACKFACING_NORMAL         42
+#define PATH_TRACER_DEBUG_VIZ_HAIR_PRIMITIVE_UV         43
diff --git a/Engine/Source/Runtime/D3D12RHI/Private/D3D12CommandContext.h b/Engine/Source/Runtime/D3D12RHI/Private/D3D12CommandContext.h
--- a/Engine/Source/Runtime/D3D12RHI/Private/D3D12CommandContext.h
+++ b/Engine/Source/Runtime/D3D12RHI/Private/D3D12CommandContext.h
@@ -743,6 +743,9 @@ public:
 #if D3D12_RHI_RAYTRACING
 	virtual void RHIBindAccelerationStructureMemory(FRHIRayTracingScene* Scene, FRHIBuffer* Buffer, uint32 BufferOffset) final override;
 	virtual void BuildAccelerationStructuresInternal(TConstArrayView<D3D12_BUILD_RAYTRACING_ACCELERATION_STRUCTURE_DESC> BuildDesc);
+#if WITH_NVAPI
+	virtual void BuildAccelerationStructuresInternalEx(const TArrayView<NVAPI_D3D12_BUILD_RAYTRACING_ACCELERATION_STRUCTURE_DESC_EX> ParamsArray);
+#endif
 #if WITH_MGPU
 	// Should be called before RHIBuildAccelerationStructures when multiple GPU support is present (for example, from FD3D12CommandContextRedirector::RHIBuildAccelerationStructures)
 	static void UnregisterAccelerationStructuresInternalMGPU(TConstArrayView<FRayTracingGeometryBuildParams> Params, FRHIGPUMask GPUMask);
diff --git a/Engine/Source/Runtime/D3D12RHI/Private/D3D12RayTracing.cpp b/Engine/Source/Runtime/D3D12RHI/Private/D3D12RayTracing.cpp
--- a/Engine/Source/Runtime/D3D12RHI/Private/D3D12RayTracing.cpp
+++ b/Engine/Source/Runtime/D3D12RHI/Private/D3D12RayTracing.cpp
@@ -123,6 +123,7 @@ DECLARE_DWORD_ACCUMULATOR_STAT(TEXT("Allocated bottom level acceleration structu
 DECLARE_DWORD_ACCUMULATOR_STAT(TEXT("Allocated top level acceleration structures"), STAT_D3D12RayTracingAllocatedTLAS, STATGROUP_D3D12RayTracing);
 DECLARE_DWORD_ACCUMULATOR_STAT(TEXT("Triangles in all BL acceleration structures"), STAT_D3D12RayTracingTrianglesBLAS, STATGROUP_D3D12RayTracing);
 DECLARE_DWORD_ACCUMULATOR_STAT(TEXT("Procedural primitives in all BL acceleration structures"), STAT_D3D12RayTracingProceduralPrimitivesBLAS, STATGROUP_D3D12RayTracing);
+DECLARE_DWORD_ACCUMULATOR_STAT(TEXT("LinearSweptSpheres in all BL acceleration structures"), STAT_D3D12RayTracingLinearSweptSpheresBLAS, STATGROUP_D3D12RayTracing);
 
 DECLARE_DWORD_COUNTER_STAT(TEXT("Built BL AS (per frame)"), STAT_D3D12RayTracingBuiltBLAS, STATGROUP_D3D12RayTracing);
 DECLARE_DWORD_COUNTER_STAT(TEXT("Updated BL AS (per frame)"), STAT_D3D12RayTracingUpdatedBLAS, STATGROUP_D3D12RayTracing);
@@ -151,6 +152,10 @@ DECLARE_CYCLE_STAT(TEXT("ClusterOperation"), STAT_D3D12ClusterOperation, STATGRO
 
 static ERayTracingAccelerationStructureFlags GetRayTracingAccelerationStructureBuildFlags(const FRayTracingGeometryInitializer& Initializer);
 
+#if WITH_NVAPI
+DEFINE_ENUM_FLAG_OPERATORS(NVAPI_D3D12_RAYTRACING_ACCELERATION_STRUCTURE_BUILD_FLAGS_EX);
+#endif
+
 #if UE_BUILD_SHIPPING
 inline void RegisterD3D12RayTracingGeometry(FD3D12RayTracingGeometry* Geometry) {};
 inline void UnregisterD3D12RayTracingGeometry(FD3D12RayTracingGeometry* Geometry) {};
@@ -435,22 +440,6 @@ static D3D12_RAYTRACING_ACCELERATION_STRUCTURE_BUILD_FLAGS TranslateRayTracingAc
 	return D3D12_RAYTRACING_ACCELERATION_STRUCTURE_BUILD_FLAGS(Result);
 }
 
-static D3D12_RAYTRACING_GEOMETRY_TYPE TranslateRayTracingGeometryType(ERayTracingGeometryType GeometryType)
-{
-	switch (GeometryType)
-	{
-	case RTGT_Triangles:
-		return D3D12_RAYTRACING_GEOMETRY_TYPE_TRIANGLES;
-		break;
-	case RTGT_Procedural:
-		return D3D12_RAYTRACING_GEOMETRY_TYPE_PROCEDURAL_PRIMITIVE_AABBS;
-		break;
-	default:
-		checkf(false, TEXT("Unexpected ray tracing geometry type"));
-		return D3D12_RAYTRACING_GEOMETRY_TYPE(0);
-	}
-}
-
 static DXGI_FORMAT TranslateRayTracingVertexElementType(EVertexElementType Type)
 {
 	switch (Type)
@@ -3301,11 +3290,15 @@ static ERayTracingAccelerationStructureFlags GetRayTracingAccelerationStructureB
 	return BuildFlags;
 }
 
-void TranslateRayTracingGeometryDescs(const FRayTracingGeometryInitializer& Initializer, TArrayView<D3D12_RAYTRACING_GEOMETRY_DESC> Output)
+static bool GeometryUsesLinearSweptSpheres(const FRayTracingGeometryInitializer& Initializer)
 {
-	check(Output.Num() == Initializer.Segments.Num());
+	return Initializer.GeometryType == RTGT_LinearSweptSpheres;
+}
 
-	D3D12_RAYTRACING_GEOMETRY_TYPE GeometryType = TranslateRayTracingGeometryType(Initializer.GeometryType);
+static void TranslateRayTracingGeometryDescs(const FRayTracingGeometryInitializer& Initializer, TArrayView<D3D12_RAYTRACING_GEOMETRY_DESC> Output)
+{
+	check(Output.Num() == Initializer.Segments.Num());
+	check(!GeometryUsesLinearSweptSpheres(Initializer));
 
 	uint32 ComputedPrimitiveCountForValidation = 0;
 
@@ -3341,12 +3334,12 @@ void TranslateRayTracingGeometryDescs(const FRayTracingGeometryInitializer& Init
 			{
 				uint32 IndexStride = Initializer.IndexBuffer->GetStride();
 				check(Initializer.IndexBuffer->GetSize() >=
-					  (Segment.FirstPrimitive + Segment.NumPrimitives) * FD3D12RayTracingGeometry::IndicesPerPrimitive * IndexStride + Initializer.IndexBufferOffset);
+					  (Segment.FirstPrimitive + Segment.NumPrimitives) * FD3D12RayTracingGeometry::IndicesPerTriangle * IndexStride + Initializer.IndexBufferOffset);
 
 				// In some cases the geometry is created with 16 bit index buffer, but it's 32 bit at build time.
 				// We conservatively set this to 32 bit to allocate acceleration structure memory.
 				Desc.Triangles.IndexFormat = DXGI_FORMAT_R32_UINT;
-				Desc.Triangles.IndexCount = Segment.NumPrimitives * FD3D12RayTracingGeometry::IndicesPerPrimitive;
+				Desc.Triangles.IndexCount = Segment.NumPrimitives * FD3D12RayTracingGeometry::IndicesPerTriangle;
 				Desc.Triangles.VertexCount = Segment.MaxVertices;
 			}
 			else
@@ -3371,9 +3364,66 @@ void TranslateRayTracingGeometryDescs(const FRayTracingGeometryInitializer& Init
 
 		ComputedPrimitiveCountForValidation += Segment.NumPrimitives;
 	}
+	check(ComputedPrimitiveCountForValidation == Initializer.TotalPrimitiveCount);
+}
+#if WITH_NVAPI
+static void TranslateRayTracingGeometryDescsEx(const FRayTracingGeometryInitializer& Initializer, TArrayView<NVAPI_D3D12_RAYTRACING_GEOMETRY_DESC_EX> Output)
+{
+	check(GeometryUsesLinearSweptSpheres(Initializer)); // This function is only for LSS primitives, otherwise use main D3D path
+	check(Output.Num() == Initializer.Segments.Num());
+
+	uint32 ComputedPrimitiveCountForValidation = 0;
+
+	for (int32 SegmentIndex = 0; SegmentIndex < Initializer.Segments.Num(); ++SegmentIndex)
+	{
+		const FRayTracingGeometrySegment& Segment = Initializer.Segments[SegmentIndex];
+
+		checkf(Segment.VertexBuffer, TEXT("Position vertex buffer is required for ray tracing geometry."));
+		checkf(Segment.VertexBufferStride, TEXT("Non-zero position vertex buffer stride is required."));
+		checkf(Segment.VertexBufferStride % 4 == 0, TEXT("Position vertex buffer stride must be aligned to 4 bytes for ByteAddressBuffer loads to work."));
+
+		checkf(Segment.MaxVertices != 0 || Segment.NumPrimitives == 0,
+			   TEXT("FRayTracingGeometrySegment.MaxVertices for '%s' must contain number of positions in the vertex buffer or maximum index buffer value+1 if index buffer is provided."),
+			   *Initializer.DebugName.ToString());
+
+		NVAPI_D3D12_RAYTRACING_GEOMETRY_DESC_EX& Desc = Output[SegmentIndex];
+		Desc = {}; // zero-initialize
+
+		Desc.flags = TranslateSegmentToGeometryFlags(Segment);
+
+		// RTGT_LinearSweptSpheres (only case for now)
+		{
+			Desc.type = NVAPI_D3D12_RAYTRACING_GEOMETRY_TYPE_LSS_EX;
+
+			checkf(Segment.VertexBufferStride >= sizeof(FVector4f), TEXT("LSS vertex buffer must contain at least float4 that defines control point."));
+
+			uint32 IndexStride = Initializer.IndexBuffer->GetStride();
+			check(Initializer.IndexBuffer->GetSize() >=
+				(Segment.FirstPrimitive + Segment.NumPrimitives) * FD3D12RayTracingGeometry::IndicesPerLinearSweptSphere * IndexStride + Initializer.IndexBufferOffset);
+
+			Desc.lss = {};
+			Desc.lss.vertexCount = Segment.MaxVertices;
+			Desc.lss.indexCount = Segment.NumPrimitives * FD3D12RayTracingGeometry::IndicesPerLinearSweptSphere;
+			Desc.lss.primitiveCount = Segment.NumPrimitives;
+			Desc.lss.vertexPositionBuffer.StartAddress = 0;
+			Desc.lss.vertexPositionBuffer.StrideInBytes = 4 * sizeof(float);
+			Desc.lss.vertexPositionFormat = DXGI_FORMAT_R32G32B32_FLOAT;
+			Desc.lss.vertexRadiusBuffer.StartAddress = 3 * sizeof(float);
+			Desc.lss.vertexRadiusBuffer.StrideInBytes = 4 * sizeof(float);
+			Desc.lss.vertexRadiusFormat = DXGI_FORMAT_R32_FLOAT;
+			Desc.lss.indexFormat = DXGI_FORMAT_R32_UINT;
+			Desc.lss.indexBuffer.StartAddress = 0;
+			Desc.lss.indexBuffer.StrideInBytes = sizeof(uint32);
+			Desc.lss.endcapMode = NVAPI_D3D12_RAYTRACING_LSS_ENDCAP_MODE_NONE;
+			Desc.lss.primitiveFormat = NVAPI_D3D12_RAYTRACING_LSS_PRIMITIVE_FORMAT_LIST;
+		}
+
+		ComputedPrimitiveCountForValidation += Segment.NumPrimitives;
+	}
 
 	check(ComputedPrimitiveCountForValidation == Initializer.TotalPrimitiveCount);
 }
+#endif // WITH_NVAPI
 
 
 FRayTracingAccelerationStructureSize FD3D12DynamicRHI::RHICalcRayTracingGeometrySize(const FRayTracingGeometryInitializer& Initializer)
@@ -3399,10 +3449,19 @@ FRayTracingAccelerationStructureSize FD3D12DynamicRHI::RHICalcRayTracingGeometry
 	D3D12_BUILD_RAYTRACING_ACCELERATION_STRUCTURE_INPUTS PrebuildDescInputs = {};
 
 	TArray<D3D12_RAYTRACING_GEOMETRY_DESC, TInlineAllocator<32>> GeometryDescs;
-	GeometryDescs.SetNumUninitialized(Initializer.Segments.Num());
-	TranslateRayTracingGeometryDescs(Initializer, GeometryDescs);
-
-	D3D12_RAYTRACING_GEOMETRY_TYPE GeometryType = TranslateRayTracingGeometryType(Initializer.GeometryType);
+#if WITH_NVAPI
+	TArray<NVAPI_D3D12_RAYTRACING_GEOMETRY_DESC_EX, TInlineAllocator<1>> GeometryDescsEx;
+	if (GeometryUsesLinearSweptSpheres(Initializer))
+	{
+		GeometryDescsEx.SetNumUninitialized(Initializer.Segments.Num());
+		TranslateRayTracingGeometryDescsEx(Initializer, GeometryDescsEx);
+	}
+	else
+#endif
+	{
+		GeometryDescs.SetNumUninitialized(Initializer.Segments.Num());
+		TranslateRayTracingGeometryDescs(Initializer, GeometryDescs);
+	}
 
 	PrebuildDescInputs.Type = D3D12_RAYTRACING_ACCELERATION_STRUCTURE_TYPE_BOTTOM_LEVEL;
 	PrebuildDescInputs.DescsLayout = D3D12_ELEMENTS_LAYOUT_ARRAY;
@@ -3430,12 +3489,45 @@ FRayTracingAccelerationStructureSize FD3D12DynamicRHI::RHICalcRayTracingGeometry
 				GeometryDesc.Triangles.IndexFormat = IndexFormat;
 			}
 		}
+#if WITH_NVAPI
+		for (NVAPI_D3D12_RAYTRACING_GEOMETRY_DESC_EX& GeometryDescEx : GeometryDescsEx)
+		{
+			check(GeometryDescEx.type == NVAPI_D3D12_RAYTRACING_GEOMETRY_TYPE_LSS_EX);
+			GeometryDescEx.lss.indexFormat = IndexFormat;
+		}
+#endif
 
 		// Get maximum buffer sizes for all GPUs in the system
 		for (uint32 GPUIndex = 0; GPUIndex < GNumExplicitGPUsForRendering; ++GPUIndex)
 		{
 			D3D12_RAYTRACING_ACCELERATION_STRUCTURE_PREBUILD_INFO PrebuildInfo = {};
-			Adapter.GetDevice(GPUIndex)->GetRaytracingAccelerationStructurePrebuildInfo(&PrebuildDescInputs, &PrebuildInfo);
+
+#if WITH_NVAPI
+			if (GeometryUsesLinearSweptSpheres(Initializer))
+			{
+				NVAPI_D3D12_BUILD_RAYTRACING_ACCELERATION_STRUCTURE_INPUTS_EX InputsEx = {};
+				InputsEx.type = PrebuildDescInputs.Type;
+				InputsEx.flags = static_cast<NVAPI_D3D12_RAYTRACING_ACCELERATION_STRUCTURE_BUILD_FLAGS_EX>(PrebuildDescInputs.Flags);
+				InputsEx.descsLayout = PrebuildDescInputs.DescsLayout;
+				InputsEx.geometryDescStrideInBytes = sizeof(NVAPI_D3D12_RAYTRACING_GEOMETRY_DESC_EX);
+				InputsEx.instanceDescs = 0ULL;
+				InputsEx.numDescs = GeometryDescsEx.Num();
+				InputsEx.pGeometryDescs = GeometryDescsEx.GetData();
+
+				NVAPI_GET_RAYTRACING_ACCELERATION_STRUCTURE_PREBUILD_INFO_EX_PARAMS PrebuildParams = {};
+				PrebuildParams.version = NVAPI_GET_RAYTRACING_ACCELERATION_STRUCTURE_PREBUILD_INFO_EX_PARAMS_VER;
+				PrebuildParams.pDesc = &InputsEx;
+				PrebuildParams.pInfo = &PrebuildInfo;
+
+				ID3D12Device5* RayTracingDevice = Adapter.GetD3DDevice5();
+				NvAPI_Status status = NvAPI_D3D12_GetRaytracingAccelerationStructurePrebuildInfoEx(RayTracingDevice, &PrebuildParams);
+				check(status == NVAPI_OK);
+			}
+			else
+#endif
+			{
+				Adapter.GetDevice(GPUIndex)->GetRaytracingAccelerationStructurePrebuildInfo(&PrebuildDescInputs, &PrebuildInfo);
+			}
 
 			SizeInfo.ResultSize = FMath::Max(SizeInfo.ResultSize, PrebuildInfo.ResultDataMaxSizeInBytes);
 			SizeInfo.BuildScratchSize = FMath::Max(SizeInfo.BuildScratchSize, PrebuildInfo.ScratchDataSizeInBytes);
@@ -3642,6 +3734,9 @@ FD3D12RayTracingGeometry::FD3D12RayTracingGeometry(FRHICommandListBase& RHICmdLi
 		SizeInfo.UpdateScratchSize = 0;
 
 		AccelerationStructureCompactedSize = OfflineBvhMetadata.Size;
+#if WITH_NVAPI
+		check(!GeometryUsesLinearSweptSpheres(Initializer));
+#endif
 	}
 	else
 	{
@@ -3681,8 +3776,9 @@ FD3D12RayTracingGeometry::FD3D12RayTracingGeometry(FRHICommandListBase& RHICmdLi
 
 	switch (Initializer.GeometryType)
 	{
-		case RTGT_Triangles: 	INC_DWORD_STAT_BY(STAT_D3D12RayTracingTrianglesBLAS				, Initializer.TotalPrimitiveCount); break;
-		case RTGT_Procedural:	INC_DWORD_STAT_BY(STAT_D3D12RayTracingProceduralPrimitivesBLAS	, Initializer.TotalPrimitiveCount); break;
+		case RTGT_Triangles:          INC_DWORD_STAT_BY(STAT_D3D12RayTracingTrianglesBLAS           , Initializer.TotalPrimitiveCount); break;
+		case RTGT_Procedural:         INC_DWORD_STAT_BY(STAT_D3D12RayTracingProceduralPrimitivesBLAS, Initializer.TotalPrimitiveCount); break;
+		case RTGT_LinearSweptSpheres: INC_DWORD_STAT_BY(STAT_D3D12RayTracingLinearSweptSpheresBLAS  , Initializer.TotalPrimitiveCount); break;
 	}
 
 	const bool bForRendering = Initializer.Type == ERayTracingGeometryInitializerType::Rendering;
@@ -3809,8 +3905,9 @@ void FD3D12RayTracingGeometry::ReleaseUnderlyingResource()
 	{
 		switch (Initializer.GeometryType)
 		{
-			case RTGT_Triangles: 	DEC_DWORD_STAT_BY(STAT_D3D12RayTracingTrianglesBLAS				, Initializer.TotalPrimitiveCount); break;
-			case RTGT_Procedural:	DEC_DWORD_STAT_BY(STAT_D3D12RayTracingProceduralPrimitivesBLAS	, Initializer.TotalPrimitiveCount); break;
+			case RTGT_Triangles:          DEC_DWORD_STAT_BY(STAT_D3D12RayTracingTrianglesBLAS           , Initializer.TotalPrimitiveCount); break;
+			case RTGT_Procedural:         DEC_DWORD_STAT_BY(STAT_D3D12RayTracingProceduralPrimitivesBLAS, Initializer.TotalPrimitiveCount); break;
+			case RTGT_LinearSweptSpheres: DEC_DWORD_STAT_BY(STAT_D3D12RayTracingLinearSweptSpheresBLAS  , Initializer.TotalPrimitiveCount); break;
 		}
 
 		DEC_DWORD_STAT(STAT_D3D12RayTracingAllocatedBLAS);
@@ -3849,7 +3946,7 @@ bool FD3D12RayTracingGeometry::AllocateBufferSRVs(uint32 InGPUIndex)
 	bool bAnyBuffersUpdated = false;
 
 	// Procedural doesn't need any SRVs for index buffer
-	if (Initializer.IndexBuffer && Initializer.GeometryType == RTGT_Triangles
+	if (Initializer.IndexBuffer && (Initializer.GeometryType == RTGT_Triangles || Initializer.GeometryType == RTGT_LinearSweptSpheres)
 #if !DO_CHECK
 		// if checks are not enabled then simply early out when HitGroupSystemIndexBufferSRV is already created because we don't support updating them
 		&& HitGroupSystemIndexBufferSRV[InGPUIndex] == nullptr		
@@ -3857,7 +3954,7 @@ bool FD3D12RayTracingGeometry::AllocateBufferSRVs(uint32 InGPUIndex)
 		)
 	{
 		checkf((Initializer.IndexBufferOffset % RHI_RAW_VIEW_ALIGNMENT) == 0, TEXT("The byte offset of raw views must be a multiple of %d (specified offset: %d)."), RHI_RAW_VIEW_ALIGNMENT, Initializer.IndexBufferOffset);
-
+		
 		FD3D12Buffer* IndexBuffer = FD3D12DynamicRHI::ResourceCast(Initializer.IndexBuffer.GetReference());
 
 		// Initializer.TotalPrimitiveCount is the accumulated num primitives of the segments
@@ -3869,7 +3966,7 @@ bool FD3D12RayTracingGeometry::AllocateBufferSRVs(uint32 InGPUIndex)
 			MaxPrimitiveCount = FMath::Max(MaxPrimitiveCount, Segment.FirstPrimitive + Segment.NumPrimitives);
 		}
 
-		const uint32 PerPrimitiveIndexCount = 3;
+		const uint32 PerPrimitiveIndexCount = (Initializer.GeometryType == RTGT_Triangles) ? IndicesPerTriangle : IndicesPerLinearSweptSphere;
 
 		D3D12_SHADER_RESOURCE_VIEW_DESC SRVDesc = {};
 		SRVDesc.ViewDimension = D3D12_SRV_DIMENSION_BUFFER;
@@ -4098,7 +4195,7 @@ void FD3D12RayTracingGeometry::SetupHitGroupSystemParameters(uint32 InGPUIndex)
 
 		FD3D12Buffer* IndexBuffer = FD3D12DynamicRHI::ResourceCast(Initializer.IndexBuffer.GetReference(), InGPUIndex);
 		const uint32 IndexStride = IndexBuffer ? IndexBuffer->GetStride() : 0;
-		const uint32 LocalIndicesPerPrimitive = FD3D12RayTracingGeometry::IndicesPerPrimitive;
+		const uint32 LocalIndicesPerPrimitive = Initializer.GeometryType == RTGT_LinearSweptSpheres ? IndicesPerLinearSweptSphere : IndicesPerTriangle;
 		for (int32 SegmentIndex = 0; SegmentIndex < Initializer.Segments.Num(); ++SegmentIndex)
 		{
 			const FRayTracingGeometrySegment& Segment = Initializer.Segments[SegmentIndex];
@@ -4139,6 +4236,8 @@ void FD3D12RayTracingGeometry::SetupHitGroupSystemParameters(uint32 InGPUIndex)
 
 void FD3D12RayTracingGeometry::CreateAccelerationStructureBuildDesc(FD3D12CommandContext& CommandContext, EAccelerationStructureBuildMode BuildMode, D3D12_GPU_VIRTUAL_ADDRESS ScratchBufferAddress, D3D12_BUILD_RAYTRACING_ACCELERATION_STRUCTURE_DESC& OutDesc, TArrayView<D3D12_RAYTRACING_GEOMETRY_DESC>& OutGeometryDescs) const
 {
+	check(!GeometryUsesLinearSweptSpheres(Initializer));
+
 	if (Initializer.IndexBuffer)
 	{
 		checkf(Initializer.IndexBuffer->GetStride() == 2 || Initializer.IndexBuffer->GetStride() == 4, TEXT("Index buffer must be 16 or 32 bit."));
@@ -4158,7 +4257,6 @@ void FD3D12RayTracingGeometry::CreateAccelerationStructureBuildDesc(FD3D12Comman
 	check(HitGroupSystemParametersForThisGPU.Num() == Initializer.Segments.Num());
 
 	ERayTracingAccelerationStructureFlags BuildFlags = GetRayTracingAccelerationStructureBuildFlags(Initializer);
-	D3D12_RAYTRACING_GEOMETRY_TYPE GeometryType = TranslateRayTracingGeometryType(Initializer.GeometryType);
 
 	// Fill in the GeometryDesc
 	TranslateRayTracingGeometryDescs(Initializer, OutGeometryDescs);
@@ -4173,28 +4271,13 @@ void FD3D12RayTracingGeometry::CreateAccelerationStructureBuildDesc(FD3D12Comman
 		
 		FD3D12Buffer* VertexBuffer = CommandContext.RetrieveObject<FD3D12Buffer>(Segment.VertexBuffer.GetReference());
 
-		switch (GeometryType)
+		switch (Initializer.GeometryType)
 		{
-		case D3D12_RAYTRACING_GEOMETRY_TYPE_TRIANGLES:
-			switch (Segment.VertexBufferElementType)
-			{
-			case VET_Float4:
-				// While the DXGI_FORMAT_R32G32B32A32_FLOAT format is not supported by DXR, since we manually load vertex 
-				// data when we are building the BLAS, we can just rely on the vertex stride to offset the read index, 
-				// and read only the 3 vertex components, and so use the DXGI_FORMAT_R32G32B32_FLOAT vertex format
-			case VET_Float3:
-				check(Desc.Triangles.VertexFormat == DXGI_FORMAT_R32G32B32_FLOAT);
-				break;
-			case VET_Float2:
-				check(Desc.Triangles.VertexFormat == DXGI_FORMAT_R32G32_FLOAT);
-				break;
-			case VET_Half2:
-				check(Desc.Triangles.VertexFormat == DXGI_FORMAT_R16G16_FLOAT);
-				break;
-			default:
-				checkNoEntry();
-				break;
-			}
+		case RTGT_Triangles:
+			check(Desc.Triangles.VertexFormat == TranslateRayTracingVertexElementType(Segment.VertexBufferElementType));
+			// #dxr_todo UE-72160: support various vertex buffer layouts (fetch/decode based on vertex stride and format)
+			checkf(Segment.VertexBufferElementType == VET_Float3 ||
+				   Segment.VertexBufferElementType == VET_Float4, TEXT("Only VET_Float3 and Float4 are currently implemented and tested. Other formats will be supported in the future."));
 
 			if (!Segment.bEnabled)
 			{
@@ -4205,7 +4288,7 @@ void FD3D12RayTracingGeometry::CreateAccelerationStructureBuildDesc(FD3D12Comman
 
 			if (IndexBuffer)
 			{
-				check(Desc.Triangles.IndexCount <= Segment.NumPrimitives * FD3D12RayTracingGeometry::IndicesPerPrimitive);
+				check(Desc.Triangles.IndexCount <= Segment.NumPrimitives * FD3D12RayTracingGeometry::IndicesPerTriangle);
 
 				Desc.Triangles.IndexFormat = (IndexStride == 4 ? DXGI_FORMAT_R32_UINT : DXGI_FORMAT_R16_UINT);
 				Desc.Triangles.IndexBuffer = IndexBuffer->ResourceLocation.GetGPUVirtualAddress() + SystemParameters.RootConstants.IndexBufferOffsetInBytes;
@@ -4223,7 +4306,7 @@ void FD3D12RayTracingGeometry::CreateAccelerationStructureBuildDesc(FD3D12Comman
 			Desc.Triangles.VertexBuffer.StrideInBytes = Segment.VertexBufferStride;
 			break;
 
-		case D3D12_RAYTRACING_GEOMETRY_TYPE_PROCEDURAL_PRIMITIVE_AABBS:
+		case RTGT_Procedural:
 			Desc.AABBs.AABBCount = Segment.NumPrimitives;
 			Desc.AABBs.AABBs.StartAddress = VertexBuffer->ResourceLocation.GetGPUVirtualAddress() + Segment.VertexBufferOffset;
 			Desc.AABBs.AABBs.StrideInBytes = Segment.VertexBufferStride;
@@ -4233,12 +4316,6 @@ void FD3D12RayTracingGeometry::CreateAccelerationStructureBuildDesc(FD3D12Comman
 			checkf(false, TEXT("Unexpected ray tracing geometry type"));
 			break;
 		}
-
-		if (GeometryType == D3D12_RAYTRACING_GEOMETRY_TYPE_TRIANGLES)
-		{
-			// #dxr_todo UE-72160: support various vertex buffer layouts (fetch/decode based on vertex stride and format)
-			checkf(Segment.VertexBufferElementType == VET_Float3 || Segment.VertexBufferElementType == VET_Float4, TEXT("Only VET_Float3 and Float4 are currently implemented and tested. Other formats will be supported in the future."));
-		}
 	}
 
 	D3D12_RAYTRACING_ACCELERATION_STRUCTURE_BUILD_FLAGS LocalBuildFlags = TranslateRayTracingAccelerationStructureFlags(BuildFlags);
@@ -4288,6 +4365,113 @@ void FD3D12RayTracingGeometry::CreateAccelerationStructureBuildDesc(FD3D12Comman
 	OutDesc = BuildDesc;
 }
 
+#if WITH_NVAPI
+void FD3D12RayTracingGeometry::CreateAccelerationStructureBuildDescEx(FD3D12CommandContext& CommandContext, EAccelerationStructureBuildMode BuildMode, D3D12_GPU_VIRTUAL_ADDRESS ScratchBufferAddress, NVAPI_D3D12_BUILD_RAYTRACING_ACCELERATION_STRUCTURE_DESC_EX& OutDescEx, TArrayView<NVAPI_D3D12_RAYTRACING_GEOMETRY_DESC_EX>& OutGeometryDescsEx) const
+{
+	check(GeometryUsesLinearSweptSpheres(GetInitializer()));
+	if (Initializer.IndexBuffer)
+	{
+		checkf(Initializer.IndexBuffer->GetStride() == 2 || Initializer.IndexBuffer->GetStride() == 4, TEXT("Index buffer must be 16 or 32 bit."));
+	}
+
+	const uint32 GPUIndex = CommandContext.GetGPUIndex();
+	const uint32 IndexStride = Initializer.IndexBuffer ? Initializer.IndexBuffer->GetStride() : 0;
+	const bool bIsUpdate = BuildMode == EAccelerationStructureBuildMode::Update;
+
+	// Make sure we've reserved the right amount of space
+	check(OutGeometryDescsEx.Num() == Initializer.Segments.Num());
+	checkf(BuffersValid(GPUIndex), TEXT("Index & vertex buffers should be valid (not streamed out) when building the acceleration structure"));
+
+	FD3D12Buffer* IndexBuffer = CommandContext.RetrieveObject<FD3D12Buffer>(Initializer.IndexBuffer.GetReference());
+		
+	const TArray<FD3D12HitGroupSystemParameters>& HitGroupSystemParametersForThisGPU = HitGroupSystemParameters[GPUIndex];
+	check(HitGroupSystemParametersForThisGPU.Num() == Initializer.Segments.Num());
+
+	TranslateRayTracingGeometryDescsEx(Initializer, OutGeometryDescsEx);
+
+	ERayTracingAccelerationStructureFlags BuildFlags = GetRayTracingAccelerationStructureBuildFlags(Initializer);
+	for (int32 SegmentIndex = 0; SegmentIndex < Initializer.Segments.Num(); ++SegmentIndex)
+	{
+		NVAPI_D3D12_RAYTRACING_GEOMETRY_DESC_EX& Desc = OutGeometryDescsEx[SegmentIndex];
+		const FRayTracingGeometrySegment& Segment = Initializer.Segments[SegmentIndex];
+		const FD3D12HitGroupSystemParameters& SystemParameters = HitGroupSystemParametersForThisGPU[SegmentIndex];
+
+		FD3D12Buffer* VertexBuffer = CommandContext.RetrieveObject<FD3D12Buffer>(Segment.VertexBuffer.GetReference());
+
+		check(Desc.type == NVAPI_D3D12_RAYTRACING_GEOMETRY_TYPE_LSS_EX);
+
+		Desc.lss.indexCount = Segment.bEnabled ? Segment.NumPrimitives * IndicesPerLinearSweptSphere : 0;
+
+		check(IndexBuffer != nullptr);
+		
+		Desc.lss.indexFormat = (IndexStride == 4 ? DXGI_FORMAT_R32_UINT : DXGI_FORMAT_R16_UINT);
+		Desc.lss.indexBuffer.StartAddress = IndexBuffer->ResourceLocation.GetGPUVirtualAddress() + SystemParameters.RootConstants.IndexBufferOffsetInBytes;
+		Desc.lss.indexBuffer.StrideInBytes = sizeof(uint32);
+		Desc.lss.vertexPositionBuffer.StartAddress = VertexBuffer->ResourceLocation.GetGPUVirtualAddress() + Segment.VertexBufferOffset;
+		Desc.lss.vertexPositionBuffer.StrideInBytes = Segment.VertexBufferStride;
+		Desc.lss.vertexRadiusBuffer.StartAddress = Desc.lss.vertexPositionBuffer.StartAddress + 3 * sizeof(float);
+		Desc.lss.vertexRadiusBuffer.StrideInBytes = Segment.VertexBufferStride;
+	}
+
+	NVAPI_D3D12_RAYTRACING_ACCELERATION_STRUCTURE_BUILD_FLAGS_EX LocalBuildFlagsEx = static_cast<NVAPI_D3D12_RAYTRACING_ACCELERATION_STRUCTURE_BUILD_FLAGS_EX>(TranslateRayTracingAccelerationStructureFlags(BuildFlags));
+
+	if (bIsUpdate)
+	{
+		checkf(EnumHasAllFlags(BuildFlags, ERayTracingAccelerationStructureFlags::AllowUpdate),
+			   TEXT("Acceleration structure must be created with FRayTracingGeometryInitializer::bAllowUpdate=true to perform refit / update."));
+
+		LocalBuildFlagsEx |= NVAPI_D3D12_RAYTRACING_ACCELERATION_STRUCTURE_BUILD_FLAG_PERFORM_UPDATE_EX;
+	}
+
+#if DO_CHECK
+	D3D12_RAYTRACING_ACCELERATION_STRUCTURE_PREBUILD_INFO PrebuildInfo = {};
+
+	NVAPI_D3D12_BUILD_RAYTRACING_ACCELERATION_STRUCTURE_INPUTS_EX InputsEx = {};
+	InputsEx.type = D3D12_RAYTRACING_ACCELERATION_STRUCTURE_TYPE_BOTTOM_LEVEL;
+	InputsEx.flags = LocalBuildFlagsEx;
+	InputsEx.descsLayout = D3D12_ELEMENTS_LAYOUT_ARRAY;
+	InputsEx.geometryDescStrideInBytes = sizeof(NVAPI_D3D12_RAYTRACING_GEOMETRY_DESC_EX);
+	InputsEx.instanceDescs = 0ULL;
+	InputsEx.numDescs = OutGeometryDescsEx.Num();
+	InputsEx.pGeometryDescs = OutGeometryDescsEx.GetData();
+
+	NVAPI_GET_RAYTRACING_ACCELERATION_STRUCTURE_PREBUILD_INFO_EX_PARAMS PrebuildParams = {};
+	PrebuildParams.version = NVAPI_GET_RAYTRACING_ACCELERATION_STRUCTURE_PREBUILD_INFO_EX_PARAMS_VER;
+	PrebuildParams.pDesc = &InputsEx;
+	PrebuildParams.pInfo = &PrebuildInfo;
+
+	ID3D12Device5* RayTracingDevice = CommandContext.GetParentDevice()->GetParentAdapter()->GetD3DDevice5();
+	NvAPI_Status status = NvAPI_D3D12_GetRaytracingAccelerationStructurePrebuildInfoEx(RayTracingDevice, &PrebuildParams);
+	check(status == NVAPI_OK);
+
+	// Must make sure that values computed in the constructor are valid.
+	check(PrebuildInfo.ResultDataMaxSizeInBytes <= SizeInfo.ResultSize);
+
+	if (bIsUpdate)
+	{
+		check(PrebuildInfo.UpdateScratchDataSizeInBytes <= SizeInfo.UpdateScratchSize);
+	}
+	else
+	{
+		check(PrebuildInfo.ScratchDataSizeInBytes <= SizeInfo.BuildScratchSize);
+	}
+#endif // DO_CHECK
+
+	OutDescEx.inputs.flags = LocalBuildFlagsEx;
+	OutDescEx.inputs.type = D3D12_RAYTRACING_ACCELERATION_STRUCTURE_TYPE_BOTTOM_LEVEL;
+	OutDescEx.inputs.numDescs = OutGeometryDescsEx.Num();
+	OutDescEx.inputs.descsLayout = D3D12_ELEMENTS_LAYOUT_ARRAY;
+	OutDescEx.inputs.geometryDescStrideInBytes = sizeof(NVAPI_D3D12_RAYTRACING_GEOMETRY_DESC_EX);
+	OutDescEx.inputs.pGeometryDescs = OutGeometryDescsEx.GetData();
+
+	OutDescEx.destAccelerationStructureData = AccelerationStructureBuffers[GPUIndex]->ResourceLocation.GetGPUVirtualAddress();
+	OutDescEx.scratchAccelerationStructureData = ScratchBufferAddress;
+	OutDescEx.sourceAccelerationStructureData = bIsUpdate
+		? AccelerationStructureBuffers[GPUIndex]->ResourceLocation.GetGPUVirtualAddress()
+		: D3D12_GPU_VIRTUAL_ADDRESS(0);
+}
+#endif
+
 static bool ShouldCompactAfterBuild(ERayTracingAccelerationStructureFlags BuildFlags)
 {
 	return EnumHasAllFlags(BuildFlags, ERayTracingAccelerationStructureFlags::AllowCompaction | ERayTracingAccelerationStructureFlags::FastTrace)
@@ -4647,6 +4831,21 @@ void FD3D12CommandContext::BuildAccelerationStructuresInternal(TConstArrayView<D
 	}
 }
 
+#if WITH_NVAPI
+void FD3D12CommandContext::BuildAccelerationStructuresInternalEx(const TArrayView<NVAPI_D3D12_BUILD_RAYTRACING_ACCELERATION_STRUCTURE_DESC_EX> DescArray)
+{
+	for (const NVAPI_D3D12_BUILD_RAYTRACING_ACCELERATION_STRUCTURE_DESC_EX& Desc : DescArray)
+	{
+		NVAPI_BUILD_RAYTRACING_ACCELERATION_STRUCTURE_EX_PARAMS Params = { 0 };
+		Params.version = NVAPI_BUILD_RAYTRACING_ACCELERATION_STRUCTURE_EX_PARAMS_VER;
+		Params.pDesc = &Desc;
+
+		NvAPI_Status status = NvAPI_D3D12_BuildRaytracingAccelerationStructureEx(GraphicsCommandList4().Get(), &Params);
+		check(status == NVAPI_OK);
+	}
+}
+#endif
+
 #if WITH_MGPU
 void FD3D12CommandContext::UnregisterAccelerationStructuresInternalMGPU(TConstArrayView<FRayTracingGeometryBuildParams> Params, FRHIGPUMask GPUMask)
 {
@@ -4709,6 +4908,11 @@ void FD3D12CommandContext::RHIBuildAccelerationStructures(TConstArrayView<FRayTr
 	TArray<D3D12_BUILD_RAYTRACING_ACCELERATION_STRUCTURE_DESC, TInlineAllocator<32>> BuildDescs;
 	BuildDescs.Reserve(Params.Num());
 
+#if WITH_NVAPI
+	TArray<NVAPI_D3D12_BUILD_RAYTRACING_ACCELERATION_STRUCTURE_DESC_EX, TInlineAllocator<32>> BuildDescsEx;
+	BuildDescsEx.Reserve(Params.Num());
+#endif
+
 	uint32 ScratchBufferSize = ScratchBufferRange.Size ? ScratchBufferRange.Size : ScratchBufferRange.Buffer->GetSize();
 
 	checkf(ScratchBufferSize + ScratchBufferRange.Offset <= ScratchBufferRange.Buffer->GetSize(),
@@ -4753,13 +4957,27 @@ void FD3D12CommandContext::RHIBuildAccelerationStructures(TConstArrayView<FRayTr
 				TEXT("BLAS scratch buffer (plus offset) must be aligned to %lld bytes."),
 				GRHIRayTracingAccelerationStructureAlignment);
 
-			// We need to keep D3D12_RAYTRACING_GEOMETRY_DESCs that are used in D3D12_BUILD_RAYTRACING_ACCELERATION_STRUCTURE_DESC alive.
+#if WITH_NVAPI
+			if (GeometryUsesLinearSweptSpheres(Geometry->GetInitializer()))
+			{
+				const uint32 NumGeometryDescs = Geometry->GetInitializer().Segments.Num();
+				NVAPI_D3D12_RAYTRACING_GEOMETRY_DESC_EX* LocalGeometryDescsMemoryEx = (NVAPI_D3D12_RAYTRACING_GEOMETRY_DESC_EX*)FMemStack::Get().Alloc(NumGeometryDescs * sizeof(NVAPI_D3D12_RAYTRACING_GEOMETRY_DESC_EX), alignof(NVAPI_D3D12_RAYTRACING_GEOMETRY_DESC_EX));
+				TArrayView<NVAPI_D3D12_RAYTRACING_GEOMETRY_DESC_EX> LocalGeometryDescsEx = MakeArrayView(LocalGeometryDescsMemoryEx, NumGeometryDescs);
+
+				NVAPI_D3D12_BUILD_RAYTRACING_ACCELERATION_STRUCTURE_DESC_EX& BuildDescEx = BuildDescsEx.AddZeroed_GetRef();
+				Geometry->CreateAccelerationStructureBuildDescEx(*this, P.BuildMode, ScratchBufferAddress, BuildDescEx, LocalGeometryDescsEx);
+			}
+			else
+#endif
+			{
+				// We need to keep D3D12_RAYTRACING_GEOMETRY_DESCs that are used in D3D12_BUILD_RAYTRACING_ACCELERATION_STRUCTURE_DESC alive.
 			const uint32 NumGeometryDescs = Geometry->GetInitializer().Segments.Num();
-			D3D12_RAYTRACING_GEOMETRY_DESC* LocalGeometryDescsMemory = (D3D12_RAYTRACING_GEOMETRY_DESC*)FMemStack::Get().Alloc(NumGeometryDescs * sizeof(D3D12_RAYTRACING_GEOMETRY_DESC), alignof(D3D12_RAYTRACING_GEOMETRY_DESC));
-			TArrayView<D3D12_RAYTRACING_GEOMETRY_DESC> LocalGeometryDescs = MakeArrayView(LocalGeometryDescsMemory, NumGeometryDescs);
+				D3D12_RAYTRACING_GEOMETRY_DESC* LocalGeometryDescsMemory = (D3D12_RAYTRACING_GEOMETRY_DESC*)FMemStack::Get().Alloc(NumGeometryDescs * sizeof(D3D12_RAYTRACING_GEOMETRY_DESC), alignof(D3D12_RAYTRACING_GEOMETRY_DESC));
+				TArrayView<D3D12_RAYTRACING_GEOMETRY_DESC> LocalGeometryDescs = MakeArrayView(LocalGeometryDescsMemory, NumGeometryDescs);
 
-			D3D12_BUILD_RAYTRACING_ACCELERATION_STRUCTURE_DESC& BuildDesc = BuildDescs.AddZeroed_GetRef();			
-			Geometry->CreateAccelerationStructureBuildDesc(*this, P.BuildMode, ScratchBufferAddress, BuildDesc, LocalGeometryDescs);
+				D3D12_BUILD_RAYTRACING_ACCELERATION_STRUCTURE_DESC& BuildDesc = BuildDescs.AddZeroed_GetRef();			
+				Geometry->CreateAccelerationStructureBuildDesc(*this, P.BuildMode, ScratchBufferAddress, BuildDesc, LocalGeometryDescs);
+			}
 
 			Geometry->UpdateResidency(*this);
 
@@ -4785,7 +5003,11 @@ void FD3D12CommandContext::RHIBuildAccelerationStructures(TConstArrayView<FRayTr
 		}
 	}
 
+
 	BuildAccelerationStructuresInternal(BuildDescs);
+#if WITH_NVAPI
+	BuildAccelerationStructuresInternalEx(BuildDescsEx);
+#endif
 
 	for (const FRayTracingGeometryBuildParams& P : Params)
 	{
diff --git a/Engine/Source/Runtime/D3D12RHI/Private/D3D12RayTracing.h b/Engine/Source/Runtime/D3D12RHI/Private/D3D12RayTracing.h
--- a/Engine/Source/Runtime/D3D12RHI/Private/D3D12RayTracing.h
+++ b/Engine/Source/Runtime/D3D12RHI/Private/D3D12RayTracing.h
@@ -149,8 +149,12 @@ public:
 	void UpdateResidency(FD3D12CommandContext& CommandContext);
 	void CompactAccelerationStructure(FD3D12CommandContext& CommandContext, uint32 InGPUIndex, uint64 InSizeAfterCompaction);
 	void CreateAccelerationStructureBuildDesc(FD3D12CommandContext& CommandContext, EAccelerationStructureBuildMode BuildMode, D3D12_GPU_VIRTUAL_ADDRESS ScratchBufferAddress,
-											D3D12_BUILD_RAYTRACING_ACCELERATION_STRUCTURE_DESC& OutDesc, TArrayView<D3D12_RAYTRACING_GEOMETRY_DESC>& OutGeometryDescs) const;
-	
+		D3D12_BUILD_RAYTRACING_ACCELERATION_STRUCTURE_DESC& OutDesc, TArrayView<D3D12_RAYTRACING_GEOMETRY_DESC>& OutGeometryDescs) const;
+#if WITH_NVAPI
+	void CreateAccelerationStructureBuildDescEx(FD3D12CommandContext& CommandContext, EAccelerationStructureBuildMode BuildMode, D3D12_GPU_VIRTUAL_ADDRESS ScratchBufferAddress,
+		NVAPI_D3D12_BUILD_RAYTRACING_ACCELERATION_STRUCTURE_DESC_EX& OutDescEx, TArrayView<NVAPI_D3D12_RAYTRACING_GEOMETRY_DESC_EX>& OutGeometryDescsEx) const;
+#endif
+
 	// Implement FD3D12ShaderResourceRenameListener interface
 	virtual void ResourceRenamed(FD3D12ContextArray const& Contexts, FD3D12BaseShaderResource* InRenamedResource, FD3D12ResourceLocation* InNewResourceLocation) override;
 
@@ -180,7 +184,8 @@ public:
 	using FRHIRayTracingGeometry::Initializer;
 	using FRHIRayTracingGeometry::SizeInfo;
 
-	static constexpr uint32 IndicesPerPrimitive = 3; // Triangle geometry only
+	static constexpr uint32 IndicesPerTriangle = 3;
+	static constexpr uint32 IndicesPerLinearSweptSphere = 2;
 
 	TRefCountPtr<FD3D12Buffer> AccelerationStructureBuffers[MAX_NUM_GPUS];
 
diff --git a/Engine/Source/Runtime/D3D12RHI/Private/Windows/WindowsD3D12Device.cpp b/Engine/Source/Runtime/D3D12RHI/Private/Windows/WindowsD3D12Device.cpp
--- a/Engine/Source/Runtime/D3D12RHI/Private/Windows/WindowsD3D12Device.cpp
+++ b/Engine/Source/Runtime/D3D12RHI/Private/Windows/WindowsD3D12Device.cpp
@@ -1710,6 +1710,28 @@ void FD3D12DynamicRHI::Init()
 				GRHIGlobals.RayTracing.SupportsClusterOps = false;
 				UE_LOG(LogD3D12RHI, Log, TEXT("Skipped NVAPI RT queries since we the feature level is below SM6"));
 			}
+
+			// LSS primitive support
+			GRHIGlobals.RayTracing.SupportsLinearSweptSpheres = false;
+			if (GRHISupportsRayTracing)
+			{
+				NVAPI_D3D12_RAYTRACING_LINEAR_SWEPT_SPHERES_CAPS LSSCaps = NVAPI_D3D12_RAYTRACING_LINEAR_SWEPT_SPHERES_CAP_NONE;
+				NvAPI_Status NvStatusLSS = NvAPI_D3D12_GetRaytracingCaps(GetAdapter().GetD3DDevice5(), NVAPI_D3D12_RAYTRACING_CAPS_TYPE_LINEAR_SWEPT_SPHERES, &LSSCaps, sizeof(NVAPI_D3D12_RAYTRACING_LINEAR_SWEPT_SPHERES_CAPS));
+
+				if (NvStatusLSS == NVAPI_OK && LSSCaps != NVAPI_D3D12_RAYTRACING_LINEAR_SWEPT_SPHERES_CAP_NONE)
+				{
+					GRHIGlobals.RayTracing.SupportsLinearSweptSpheres = CheckPipelineStateCap(NVAPI_D3D12_PIPELINE_CREATION_STATE_FLAGS_ENABLE_LSS_SUPPORT);
+				}
+
+				if (GRHIGlobals.RayTracing.SupportsLinearSweptSpheres)
+				{
+					UE_LOG(LogD3D12RHI, Log, TEXT("LSS support is enabled."));
+				}
+				else
+				{
+					UE_LOG(LogD3D12RHI, Log, TEXT("LSS not supported on this device."));
+				}
+			}
 		}
 
 		// Disable ray tracing for old Nvidia drivers
@@ -1756,7 +1778,6 @@ void FD3D12DynamicRHI::Init()
 				checkf(0, TEXT("NvAPI SetCreatePipelineStateOptions failed, incompatible feature flags are enabled"));
 			}
 		}
-
 	} // if NVIDIA
 #endif // NVAPI
 
diff --git a/Engine/Source/Runtime/RHI/Private/DataDrivenShaderPlatformInfo.cpp b/Engine/Source/Runtime/RHI/Private/DataDrivenShaderPlatformInfo.cpp
--- a/Engine/Source/Runtime/RHI/Private/DataDrivenShaderPlatformInfo.cpp
+++ b/Engine/Source/Runtime/RHI/Private/DataDrivenShaderPlatformInfo.cpp
@@ -239,6 +239,7 @@ void FGenericDataDrivenShaderPlatformInfo::ParseDataDrivenShaderInfo(const FConf
 	GET_SECTION_BOOL_HELPER(bInlineRayTracingRequiresBindless);
 	GET_SECTION_BOOL_HELPER(bSupportsRayTracingCallableShaders);
 	GET_SECTION_BOOL_HELPER(bSupportsRayTracingProceduralPrimitive);
+	GET_SECTION_BOOL_HELPER(bSupportsRayTracingLinearSweptSpheres);
 	GET_SECTION_BOOL_HELPER(bSupportsRayTracingTraversalStatistics);
 	GET_SECTION_BOOL_HELPER(bSupportsRayTracingIndirectInstanceData);
 	GET_SECTION_BOOL_HELPER(bSupportsRayTracingClusterOps);
diff --git a/Engine/Source/Runtime/RHI/Public/DataDrivenShaderPlatformInfo.h b/Engine/Source/Runtime/RHI/Public/DataDrivenShaderPlatformInfo.h
--- a/Engine/Source/Runtime/RHI/Public/DataDrivenShaderPlatformInfo.h
+++ b/Engine/Source/Runtime/RHI/Public/DataDrivenShaderPlatformInfo.h
@@ -48,6 +48,7 @@ class FGenericDataDrivenShaderPlatformInfo
 	uint32 bSupportsRayTracing : 1;
 	uint32 bSupportsRayTracingCallableShaders : 1;
 	uint32 bSupportsRayTracingProceduralPrimitive : 1;
+	uint32 bSupportsRayTracingLinearSweptSpheres : 1;
 	uint32 bSupportsRayTracingTraversalStatistics : 1;
 	uint32 bSupportsRayTracingIndirectInstanceData : 1; // Whether instance transforms can be copied from the GPU to the TLAS instances buffer
 	uint32 bSupportsRayTracingClusterOps : 1; // Accelerated building and ray tracing of Nanite clusters
@@ -397,6 +398,12 @@ public:
 		return Infos[Platform].bSupportsRayTracing && Infos[Platform].bSupportsRayTracingProceduralPrimitive;
 	}
 
+	static inline const bool GetSupportsRayTracingLinearSweptSpheres(const FStaticShaderPlatform Platform)
+	{
+		check(IsValid(Platform));
+		return Infos[Platform].bSupportsRayTracing && Infos[Platform].bSupportsRayTracingLinearSweptSpheres;
+	}
+
 	static inline const bool GetSupportsRayTracingTraversalStatistics(const FStaticShaderPlatform Platform)
 	{
 		check(IsValid(Platform));
diff --git a/Engine/Source/Runtime/RHI/Public/RHIGlobals.h b/Engine/Source/Runtime/RHI/Public/RHIGlobals.h
--- a/Engine/Source/Runtime/RHI/Public/RHIGlobals.h
+++ b/Engine/Source/Runtime/RHI/Public/RHIGlobals.h
@@ -473,6 +473,9 @@ struct FRHIGlobals
 		/** Whether or not the RHI supports a separate hit group buffer when building the acceleration structure */
 		bool RequiresSeparateHitGroupContributionsBuffer = false;
 
+		/** Whether or not the RHI supports linear swept sphere primitives **/
+		bool SupportsLinearSweptSpheres = false;
+
 		/** Required alignment for ray tracing acceleration structures. */
 		uint32 AccelerationStructureAlignment = 0;
 
diff --git a/Engine/Source/Runtime/RHI/Public/RHIResources.h b/Engine/Source/Runtime/RHI/Public/RHIResources.h
--- a/Engine/Source/Runtime/RHI/Public/RHIResources.h
+++ b/Engine/Source/Runtime/RHI/Public/RHIResources.h
@@ -3402,6 +3402,12 @@ enum ERayTracingGeometryType
 	// Vertex stride must be at least 24 bytes, but may be larger to support custom per-primitive data.
 	// Index buffers can't be used with procedural geometry.
 	RTGT_Procedural,
+
+	// Linear Swept Sphere (LSS) is a ray tracing primitive that represents the volume traced by a sphere moving
+	// linearly between two endpoints, optionally with varying radii. It efficiently models strand-like geometry
+	// such as hair or fur without tessellating into triangles, offering smoother appearance, lower memory use,
+	// and hardware-accelerated ray intersections.
+	RTGT_LinearSweptSpheres,
 };
 DECLARE_INTRINSIC_TYPE_LAYOUT(ERayTracingGeometryType);
 
diff --git a/Engine/Source/Runtime/Renderer/Private/HairStrands/HairStrandsData.h b/Engine/Source/Runtime/Renderer/Private/HairStrands/HairStrandsData.h
--- a/Engine/Source/Runtime/Renderer/Private/HairStrands/HairStrandsData.h
+++ b/Engine/Source/Runtime/Renderer/Private/HairStrands/HairStrandsData.h
@@ -500,6 +500,7 @@ namespace HairStrands
 	bool HasHairInstanceInScene(const FScene& Scene);
 	bool HasHairCardsVisible(const TArray<FViewInfo>& Views);
 	bool HasHairStrandsVisible(TArrayView<const FViewInfo> Views);
+	bool IsHairStrandsVF(const FVertexFactoryType* InVFType);
 
 	void AddVisibleShadowCastingLight(const FScene& Scene, TArray<FViewInfo>& Views, const FLightSceneInfo* LightSceneInfo);
 
diff --git a/Engine/Source/Runtime/Renderer/Private/HairStrands/HairStrandsInterface.cpp b/Engine/Source/Runtime/Renderer/Private/HairStrands/HairStrandsInterface.cpp
--- a/Engine/Source/Runtime/Renderer/Private/HairStrands/HairStrandsInterface.cpp
+++ b/Engine/Source/Runtime/Renderer/Private/HairStrands/HairStrandsInterface.cpp
@@ -82,6 +82,11 @@ static TAutoConsoleVariable<int32> CVarHairStrandsVisibilityComputeRaster_Contin
 	TEXT("Enable Continuos LOD when using compute rasterization. Experimental"),
 	ECVF_RenderThreadSafe);
 
+static TAutoConsoleVariable<int32> CVarHairStrandsRaytracingEnableLinearSweptSpheres(
+	TEXT("r.HairStrands.Raytracing.EnableLinearSweptSpheres"), 1,
+	TEXT("Enable Linear Swept Spheres for hair strand raytracing (if the current hardware supports it)."),
+	ECVF_RenderThreadSafe | ECVF_ReadOnly);
+
 ////////////////////////////////////////////////////////////////////////////////////////////////////////////////////////
 // Import/export utils function for hair resources
 void FRDGExternalBuffer::Release()
@@ -215,6 +220,12 @@ bool IsHairStrandsSimulationEnable()
 	return CVarHairStrandsSimulation.GetValueOnAnyThread() > 0;
 }
 
+bool IsHairStrandsLinearSweptSpheresEnabled()
+{
+	// Are LSS primitives supported on the currently running device and are they enable for hair strands in this project?
+	return GRHIGlobals.RayTracing.SupportsLinearSweptSpheres && CVarHairStrandsRaytracingEnableLinearSweptSpheres.GetValueOnAnyThread() != 0;
+}
+
 ////////////////////////////////////////////////////////////////////////////////////////////////////////////////////////
 // FHairTransientResources
 
@@ -569,9 +580,7 @@ bool IsHairStrandsVF(const FMeshBatch* Mesh)
 {
 	if (Mesh)
 	{
-		static const FHashedName& VFTypeRef = FVertexFactoryType::GetVFByName(TEXT("FHairStrandsVertexFactory"))->GetHashedName();
-		const FHashedName& VFType = Mesh->VertexFactory->GetType()->GetHashedName();
-		return VFType == VFTypeRef;
+		return IsHairStrandsVF(Mesh->VertexFactory->GetType());
 	}
 	return false;
 }
@@ -607,6 +616,14 @@ bool IsHairVisible(const FMeshBatchAndRelevance& MeshBatch, bool bCheckLengthSca
 	return false;
 }
 
+bool IsHairStrandsVF(const FVertexFactoryType* InVFType)
+{
+	// NOTE: The VF is defined in a different module, so we must search for it by name
+	static const FVertexFactoryType* VFTypeA = FVertexFactoryType::GetVFByName(TEXT("FHairStrandsVertexFactory"));                   check(VFTypeA != nullptr);
+	static const FVertexFactoryType* VFTypeB = FVertexFactoryType::GetVFByName(TEXT("FLinearSweptSpheresHairStrandsVertexFactory")); check(VFTypeB != nullptr);
+	return InVFType == VFTypeA || InVFType == VFTypeB;
+}
+
 FHairGroupPublicData* GetHairData(const FMeshBatch* Mesh)
 {
 	return reinterpret_cast<FHairGroupPublicData*>(Mesh->Elements[0].VertexFactoryUserData);
diff --git a/Engine/Source/Runtime/Renderer/Private/HairStrands/HairStrandsRasterCommon.cpp b/Engine/Source/Runtime/Renderer/Private/HairStrands/HairStrandsRasterCommon.cpp
--- a/Engine/Source/Runtime/Renderer/Private/HairStrands/HairStrandsRasterCommon.cpp
+++ b/Engine/Source/Runtime/Renderer/Private/HairStrands/HairStrandsRasterCommon.cpp
@@ -35,7 +35,7 @@ protected:
 
 	static bool ShouldCompilePermutation(const FMeshMaterialShaderPermutationParameters& Parameters)
 	{
-		return IsCompatibleWithHairStrands(Parameters.Platform, Parameters.MaterialParameters) && Parameters.VertexFactoryType->GetFName() == FName("FHairStrandsVertexFactory");
+		return IsCompatibleWithHairStrands(Parameters.Platform, Parameters.MaterialParameters) && HairStrands::IsHairStrandsVF(Parameters.VertexFactoryType);
 	}
 
 	static void ModifyCompilationEnvironment(const FMaterialShaderPermutationParameters& Parameters, FShaderCompilerEnvironment& OutEnvironment)
@@ -67,7 +67,7 @@ protected:
 
 	static bool ShouldCompilePermutation(const FMeshMaterialShaderPermutationParameters& Parameters)
 	{
-		return IsCompatibleWithHairStrands(Parameters.Platform, Parameters.MaterialParameters) && Parameters.VertexFactoryType->GetFName() == FName("FHairStrandsVertexFactory");
+		return IsCompatibleWithHairStrands(Parameters.Platform, Parameters.MaterialParameters) && HairStrands::IsHairStrandsVF(Parameters.VertexFactoryType);
 	}
 
 	static void ModifyCompilationEnvironment(const FMaterialShaderPermutationParameters& Parameters, FShaderCompilerEnvironment& OutEnvironment)
@@ -99,7 +99,7 @@ public:
 
 	static bool ShouldCompilePermutation(const FMeshMaterialShaderPermutationParameters& Parameters)
 	{
-		return IsCompatibleWithHairStrands(Parameters.Platform, Parameters.MaterialParameters) && Parameters.VertexFactoryType->GetFName() == FName("FHairStrandsVertexFactory");
+		return IsCompatibleWithHairStrands(Parameters.Platform, Parameters.MaterialParameters) && HairStrands::IsHairStrandsVF(Parameters.VertexFactoryType);
 	}
 };
 IMPLEMENT_MATERIAL_SHADER_TYPE(, FDeepShadowDepthMeshPS, TEXT("/Engine/Private/HairStrands/HairStrandsDeepShadowPS.usf"), TEXT("MainDepth"), SF_Pixel);
@@ -123,7 +123,7 @@ public:
 
 	static bool ShouldCompilePermutation(const FMeshMaterialShaderPermutationParameters& Parameters)
 	{
-		return IsCompatibleWithHairStrands(Parameters.Platform, Parameters.MaterialParameters) && Parameters.VertexFactoryType->GetFName() == FName("FHairStrandsVertexFactory");
+		return IsCompatibleWithHairStrands(Parameters.Platform, Parameters.MaterialParameters) && HairStrands::IsHairStrandsVF(Parameters.VertexFactoryType);
 	}
 };
 IMPLEMENT_MATERIAL_SHADER_TYPE(, FDeepShadowDomMeshPS, TEXT("/Engine/Private/HairStrands/HairStrandsDeepShadowPS.usf"), TEXT("MainDom"), SF_Pixel);
@@ -236,7 +236,6 @@ bool FHairRasterMeshProcessor::Process(
 	ERasterizerCullMode MeshCullMode)
 {
 	const FVertexFactory* VertexFactory = MeshBatch.VertexFactory;
-	static const FVertexFactoryType* CompatibleVF = FVertexFactoryType::GetVFByName(TEXT("FHairStrandsVertexFactory"));
 
 	FMaterialShaderTypes ShaderTypes;
 	ShaderTypes.AddShaderType<VertexShaderType>();
@@ -253,7 +252,7 @@ bool FHairRasterMeshProcessor::Process(
 		PixelShaderType> PassShaders;
 	{
 		FVertexFactoryType* VertexFactoryType = VertexFactory->GetType();
-		const bool bIsHairStrandsFactory = MeshBatch.VertexFactory->GetType()->GetHashedName() == CompatibleVF->GetHashedName();
+		const bool bIsHairStrandsFactory = HairStrands::IsHairStrandsVF(MeshBatch.VertexFactory->GetType());
 		if (!bIsHairStrandsFactory)
 			return true;	// Skip adding this batch if the VF type does not match.
 
diff --git a/Engine/Source/Runtime/Renderer/Private/HairStrands/HairStrandsVisibility.cpp b/Engine/Source/Runtime/Renderer/Private/HairStrands/HairStrandsVisibility.cpp
--- a/Engine/Source/Runtime/Renderer/Private/HairStrands/HairStrandsVisibility.cpp
+++ b/Engine/Source/Runtime/Renderer/Private/HairStrands/HairStrandsVisibility.cpp
@@ -495,8 +495,7 @@ protected:
 
 	static bool ShouldCompilePermutation(const FMeshMaterialShaderPermutationParameters& Parameters)
 	{
-		const bool bIsCompatible = IsCompatibleWithHairVisibility(Parameters) && Parameters.VertexFactoryType->GetFName() == FName("FHairStrandsVertexFactory");
-		return bIsCompatible;
+		return IsCompatibleWithHairVisibility(Parameters) && HairStrands::IsHairStrandsVF(Parameters.VertexFactoryType);
 	}
 
 	static void ModifyCompilationEnvironment(const FMaterialShaderPermutationParameters& Parameters, FShaderCompilerEnvironment& OutEnvironment)
@@ -562,8 +561,7 @@ public:
 
 	static bool ShouldCompilePermutation(const FMeshMaterialShaderPermutationParameters& Parameters)
 	{
-		const bool bIsCompatible = IsCompatibleWithHairStrands(Parameters.Platform, Parameters.MaterialParameters) && Parameters.VertexFactoryType->GetFName() == FName("FHairStrandsVertexFactory");
-		return bIsCompatible;
+		return IsCompatibleWithHairStrands(Parameters.Platform, Parameters.MaterialParameters) && HairStrands::IsHairStrandsVF(Parameters.VertexFactoryType);
 	}
 
 	static void ModifyCompilationEnvironment(const FMaterialShaderPermutationParameters& Parameters, FShaderCompilerEnvironment& OutEnvironment)
@@ -617,11 +615,9 @@ enum class EHairMaterialPassFilter : uint8
 
 static bool ShouldRenderHairStrands(ERHIFeatureLevel::Type FeatureLevel, const FMaterial& Material, const FVertexFactoryType* VFType, bool bPrimitiveRenderInPass)
 {
-	static const FVertexFactoryType* CompatibleVF = FVertexFactoryType::GetVFByName(TEXT("FHairStrandsVertexFactory"));
-
 	// Determine the mesh's material and blend mode.
 	const bool bIsCompatible = IsCompatibleWithHairStrands(&Material, FeatureLevel);
-	const bool bIsHairStrandsFactory = CompatibleVF != nullptr && VFType->GetHashedName() == CompatibleVF->GetHashedName();
+	const bool bIsHairStrandsFactory = HairStrands::IsHairStrandsVF(VFType);
 
 	return bPrimitiveRenderInPass && bIsCompatible && bIsHairStrandsFactory && ShouldIncludeDomainInMeshPass(Material.GetMaterialDomain());
 }
@@ -1408,7 +1404,7 @@ protected:
 
 	static bool ShouldCompilePermutation(const FMeshMaterialShaderPermutationParameters& Parameters)
 	{
-		return IsCompatibleWithHairVisibility(Parameters) && Parameters.VertexFactoryType->GetFName() == FName("FHairStrandsVertexFactory");
+		return IsCompatibleWithHairStrands(Parameters.Platform, Parameters.MaterialParameters) && HairStrands::IsHairStrandsVF(Parameters.VertexFactoryType);
 	}
 
 	static void ModifyCompilationEnvironment(const FMaterialShaderPermutationParameters& Parameters, FShaderCompilerEnvironment& OutEnvironment)
@@ -1466,7 +1462,7 @@ public:
 
 	static bool ShouldCompilePermutation(const FMeshMaterialShaderPermutationParameters& Parameters)
 	{
-		if (Parameters.VertexFactoryType->GetFName() != FName("FHairStrandsVertexFactory"))
+		if (!HairStrands::IsHairStrandsVF(Parameters.VertexFactoryType))
 		{
 			return false;
 		}
diff --git a/Engine/Source/Runtime/Renderer/Private/PathTracing.cpp b/Engine/Source/Runtime/Renderer/Private/PathTracing.cpp
--- a/Engine/Source/Runtime/Renderer/Private/PathTracing.cpp
+++ b/Engine/Source/Runtime/Renderer/Private/PathTracing.cpp
@@ -104,6 +104,7 @@ struct FPathTracingDebugModeInfo
 		REGISTER_VIZ_TYPE(TEXT("EyeIrisNormal")         , PATH_TRACER_DEBUG_VIZ_EYE_IRIS_NORMAL         );
 		REGISTER_VIZ_TYPE(TEXT("EyeCausticNormal")      , PATH_TRACER_DEBUG_VIZ_EYE_CAUSTIC_NORMAL      );
 		REGISTER_VIZ_TYPE(TEXT("BackfacingNormal")      , PATH_TRACER_DEBUG_VIZ_BACKFACING_NORMAL       );
+		REGISTER_VIZ_TYPE(TEXT("HairPrimitiveUV")		, PATH_TRACER_DEBUG_VIZ_HAIR_PRIMITIVE_UV		);
 #undef REGISTER_VIZ_TYPE
 	}
 
diff --git a/Engine/Source/Runtime/Renderer/Public/HairStrandsInterface.h b/Engine/Source/Runtime/Renderer/Public/HairStrandsInterface.h
--- a/Engine/Source/Runtime/Renderer/Public/HairStrandsInterface.h
+++ b/Engine/Source/Runtime/Renderer/Public/HairStrandsInterface.h
@@ -194,6 +194,14 @@ enum EHairGeometryType
 	NoneGeometry
 };
 
+enum class EHairStrandsRayTracingPrimitiveType : uint8
+{
+	Procedural,
+	Triangle,
+	LinearSweptSphere,
+	MAX
+};
+
 enum EHairBindingType
 {
 	NoneBinding,
@@ -538,6 +546,9 @@ RENDERER_API bool IsHairStrandContinuousDecimationReorderingEnabled();
 // Return true if continuous LOD is enabled - implies computer raster and continuous decimation reordering is true
 RENDERER_API bool IsHairVisibilityComputeRasterContinuousLODEnabled();
 
+// Return true if Linear Swept Spheres enabled for hair strands
+RENDERER_API bool IsHairStrandsLinearSweptSpheresEnabled();
+
 ////////////////////////////////////////////////////////////////////////////////////////////////////////////////////////
 
 typedef TArray<FRHIUnorderedAccessView*> FBufferTransitionQueue;
diff --git a/Engine/Source/Runtime/VulkanRHI/Private/VulkanExtensions.cpp b/Engine/Source/Runtime/VulkanRHI/Private/VulkanExtensions.cpp
--- a/Engine/Source/Runtime/VulkanRHI/Private/VulkanExtensions.cpp
+++ b/Engine/Source/Runtime/VulkanRHI/Private/VulkanExtensions.cpp
@@ -1494,6 +1494,43 @@ private:
 	VkPhysicalDeviceRayTracingValidationFeaturesNV RayTracingValidationFeaturesNV;
 };
 
+// ***** VK_NV_ray_tracing_linear_swept_spheres (vendor)
+class FVulkanNVRayTracingLinearSweptSphereExtension : public FVulkanDeviceExtension
+{
+public:
+
+	FVulkanNVRayTracingLinearSweptSphereExtension(FVulkanDevice* InDevice)
+		: FVulkanDeviceExtension(InDevice, VK_NV_RAY_TRACING_LINEAR_SWEPT_SPHERES_EXTENSION_NAME, VULKAN_EXTENSION_ENABLED)
+	{
+		bEnabledInCode = bEnabledInCode && UE::RHICore::AllowVendorDevice();
+	}
+
+	virtual void PrePhysicalDeviceFeatures(VkPhysicalDeviceFeatures2KHR& PhysicalDeviceFeatures2) override final
+	{
+		ZeroVulkanStruct(LinearSweptSphereFeaturesNV, VK_STRUCTURE_TYPE_PHYSICAL_DEVICE_RAY_TRACING_LINEAR_SWEPT_SPHERES_FEATURES_NV);
+		AddToPNext(PhysicalDeviceFeatures2, LinearSweptSphereFeaturesNV);
+	}
+
+	virtual void PostPhysicalDeviceFeatures(FOptionalVulkanDeviceExtensions& ExtensionFlags) override final
+	{
+		bRequirementsPassed = (LinearSweptSphereFeaturesNV.linearSweptSpheres == VK_TRUE);
+	}
+
+	virtual void PreCreateDevice(VkDeviceCreateInfo& DeviceCreateInfo) override final
+	{
+		GRHIGlobals.RayTracing.SupportsLinearSweptSpheres = false;
+		if (bRequirementsPassed)
+		{
+			AddToPNext(DeviceCreateInfo, LinearSweptSphereFeaturesNV);
+
+			GRHIGlobals.RayTracing.SupportsLinearSweptSpheres = (LinearSweptSphereFeaturesNV.linearSweptSpheres == VK_TRUE);
+		}
+	}
+
+private:
+	VkPhysicalDeviceRayTracingLinearSweptSpheresFeaturesNV LinearSweptSphereFeaturesNV;
+};
+
 // ***** VK_EXT_device_fault
 class FVulkanEXTDeviceFaultExtension : public FVulkanDeviceExtension
 {
@@ -2690,6 +2727,7 @@ FVulkanDeviceExtensionArray FVulkanDeviceExtension::GetUESupportedDeviceExtensio
 	ADD_CUSTOM_EXTENSION(FVulkanNVDeviceDiagnosticCheckpointsExtension);
 	ADD_CUSTOM_EXTENSION(FVulkanNVDeviceDiagnosticConfigExtension);
 	ADD_CUSTOM_EXTENSION(FVulkanNVRayTracingValidationExtension);
+	ADD_CUSTOM_EXTENSION(FVulkanNVRayTracingLinearSweptSphereExtension);
 
 	// Add in platform specific extensions
 	FVulkanPlatform::GetDeviceExtensions(InDevice, OutUEDeviceExtensions);
diff --git a/Engine/Source/Runtime/VulkanRHI/Private/VulkanRayTracing.cpp b/Engine/Source/Runtime/VulkanRHI/Private/VulkanRayTracing.cpp
--- a/Engine/Source/Runtime/VulkanRHI/Private/VulkanRayTracing.cpp
+++ b/Engine/Source/Runtime/VulkanRHI/Private/VulkanRayTracing.cpp
@@ -217,6 +217,12 @@ static void GetBLASBuildData(
 			: (IndexBuffer->GetSize() > 0) ? IndexBuffer->GetStride() : 4;
 	}
 
+	// Pre allocate LSS Segments to avoid memory reallocation
+	if (GeometryType == RTGT_LinearSweptSpheres)
+	{
+		BuildData.LSSSegments.Reserve(Segments.Num());
+	}
+
 	TArray<uint32, TInlineAllocator<1>> PrimitiveCounts;
 
 	for (int32 SegmentIndex = 0; SegmentIndex < Segments.Num(); ++SegmentIndex)
@@ -244,6 +250,9 @@ static void GetBLASBuildData(
 			SegmentGeometry.flags |= VK_GEOMETRY_NO_DUPLICATE_ANY_HIT_INVOCATION_BIT_KHR;
 		}
 
+		VkAccelerationStructureGeometryLinearSweptSpheresDataNV LSSSegmentGeometry;
+		ZeroVulkanStruct(LSSSegmentGeometry, VK_STRUCTURE_TYPE_ACCELERATION_STRUCTURE_GEOMETRY_LINEAR_SWEPT_SPHERES_DATA_NV);
+
 		uint32 PrimitiveOffset = 0;
 		switch (GeometryType)
 		{
@@ -296,6 +305,36 @@ static void GetBLASBuildData(
 				SegmentGeometry.geometry.aabbs.data = VertexBufferDeviceAddress;
 				SegmentGeometry.geometry.aabbs.stride = Segment.VertexBufferStride;
 
+				break;
+
+			case RTGT_LinearSweptSpheres:
+				// no other application rather than hair strands for now.
+				check(IndexStrideInBytes == 4);
+				check(IndexBufferRHI != nullptr); // must have an index buffer in this case
+
+				// there are 2 indices per primitive
+				PrimitiveOffset = Segment.FirstPrimitive * 2 * IndexStrideInBytes;
+
+				LSSSegmentGeometry.indexType = VK_INDEX_TYPE_UINT32;
+				LSSSegmentGeometry.indexData = IndexBufferDeviceAddress;
+				LSSSegmentGeometry.indexStride = sizeof(uint32);
+				LSSSegmentGeometry.indexingMode = VkRayTracingLssIndexingModeNV::VK_RAY_TRACING_LSS_INDEXING_MODE_LIST_NV;
+				
+				LSSSegmentGeometry.vertexData = VertexBufferDeviceAddress;
+				LSSSegmentGeometry.vertexFormat = VkFormat::VK_FORMAT_R32G32B32_SFLOAT;
+				LSSSegmentGeometry.vertexStride = 4 * sizeof(float);
+				
+				LSSSegmentGeometry.radiusData = {.deviceAddress = VertexBufferDeviceAddress.deviceAddress + 3 * sizeof(float)};
+				LSSSegmentGeometry.radiusFormat = VkFormat::VK_FORMAT_R32_SFLOAT;
+				LSSSegmentGeometry.radiusStride = 4 * sizeof(float);
+
+				LSSSegmentGeometry.endCapsMode = VkRayTracingLssPrimitiveEndCapsModeNV::VK_RAY_TRACING_LSS_PRIMITIVE_END_CAPS_MODE_NONE_NV;
+				
+				BuildData.LSSSegments.Add(LSSSegmentGeometry);
+
+				SegmentGeometry.pNext = &BuildData.LSSSegments[SegmentIndex];
+				SegmentGeometry.geometryType = VkGeometryTypeKHR::VK_GEOMETRY_TYPE_LINEAR_SWEPT_SPHERES_NV;
+
 				break;
 			default:
 				checkf(false, TEXT("Unexpected ray tracing geometry type"));
@@ -1438,6 +1477,15 @@ public:
 				RayTracingPipelineCreateInfo.layout = Device.GetBindlessDescriptorManager()->GetPipelineLayout();
 				RayTracingPipelineCreateInfo.flags = VK_PIPELINE_CREATE_LIBRARY_BIT_KHR | VK_PIPELINE_CREATE_DESCRIPTOR_BUFFER_BIT_EXT;
 
+				VkPipelineCreateFlags2CreateInfoKHR PipelineFlags2;
+				ZeroVulkanStruct(PipelineFlags2, VK_STRUCTURE_TYPE_PIPELINE_CREATE_FLAGS_2_CREATE_INFO_KHR);
+				PipelineFlags2.flags = RayTracingPipelineCreateInfo.flags | VK_PIPELINE_CREATE_2_RAY_TRACING_ALLOW_SPHERES_AND_LINEAR_SWEPT_SPHERES_BIT_NV;
+
+				if (GRHIGlobals.RayTracing.SupportsLinearSweptSpheres)
+				{
+					RayTracingPipelineCreateInfo.pNext = &PipelineFlags2;
+				}
+
 				const VkResult Result = Device.GetPipelineStateCache()->CreateRayTracingPipeline(RayTracingPipelineCreateInfo, bPartial, PipelineLibrary->PipelineHandle);
 				if (!bPartial)
 				{
@@ -1605,6 +1653,15 @@ FVulkanRayTracingPipelineState::FVulkanRayTracingPipelineState(FVulkanDevice& In
 	RayTracingPipelineCreateInfo.layout = InDevice.GetBindlessDescriptorManager()->GetPipelineLayout();
 	RayTracingPipelineCreateInfo.flags = VK_PIPELINE_CREATE_DESCRIPTOR_BUFFER_BIT_EXT;
 
+	VkPipelineCreateFlags2CreateInfoKHR PipelineFlags2;
+	ZeroVulkanStruct(PipelineFlags2, VK_STRUCTURE_TYPE_PIPELINE_CREATE_FLAGS_2_CREATE_INFO_KHR);
+	PipelineFlags2.flags = RayTracingPipelineCreateInfo.flags | VK_PIPELINE_CREATE_2_RAY_TRACING_ALLOW_SPHERES_AND_LINEAR_SWEPT_SPHERES_BIT_NV;
+
+	if (GRHIGlobals.RayTracing.SupportsLinearSweptSpheres)
+	{
+		RayTracingPipelineCreateInfo.pNext = &PipelineFlags2;
+	}
+	
 	VERIFYVULKANRESULT_EXPANDED(InDevice.GetPipelineStateCache()->CreateRayTracingPipeline(RayTracingPipelineCreateInfo, bIsPartialPipeline, Pipeline));
 
 	// Grab all shader handles for each stage
diff --git a/Engine/Source/Runtime/VulkanRHI/Private/VulkanRayTracing.h b/Engine/Source/Runtime/VulkanRHI/Private/VulkanRayTracing.h
--- a/Engine/Source/Runtime/VulkanRHI/Private/VulkanRayTracing.h
+++ b/Engine/Source/Runtime/VulkanRHI/Private/VulkanRayTracing.h
@@ -60,6 +60,7 @@ struct FVkRtBLASBuildData
 	}
 
 	TArray<VkAccelerationStructureGeometryKHR, TInlineAllocator<1>> Segments;
+	TArray<VkAccelerationStructureGeometryLinearSweptSpheresDataNV, TInlineAllocator<1>> LSSSegments;
 	TArray<VkAccelerationStructureBuildRangeInfoKHR, TInlineAllocator<1>> Ranges;
 	VkAccelerationStructureBuildGeometryInfoKHR GeometryInfo;
 	VkAccelerationStructureBuildSizesInfoKHR SizesInfo;
