Add GPUCrash.D3DResult and GPUCrash.D3DDeviceRemovedReason fields to crash reports for DX11 and DX12

#Synced-CL 51619035

diff --git a/Engine/Source/Runtime/D3D12RHI/Private/D3D12RHIPrivate.h b/Engine/Source/Runtime/D3D12RHI/Private/D3D12RHIPrivate.h
--- a/Engine/Source/Runtime/D3D12RHI/Private/D3D12RHIPrivate.h
+++ b/Engine/Source/Runtime/D3D12RHI/Private/D3D12RHIPrivate.h
@@ -187,7 +187,7 @@ public:
 
 private:
 	// Calls OutputGpuCrashReport, and also forces a quit.
-	void TerminateOnGPUCrash();
+	void TerminateOnGPUCrash(TOptional<HRESULT> D3DResult);
 
 	/** Texture pool size */
 	int64 RequestedTexturePoolSize;
@@ -257,7 +257,7 @@ public:
 	void ProcessInterruptQueueUntil(FGraphEvent* GraphEvent);
 
 	// Called when the GPU has crashed. This function will not return.
-	void ProcessInterruptQueueOnGPUCrash();
+	void ProcessInterruptQueueOnGPUCrash(HRESULT D3DResult);
 
 	struct FD3D12TimingArray : public TArray<TUniquePtr<FD3D12Timing>, TInlineAllocator<GD3D12MaxNumQueues>>
 	{
diff --git a/Engine/Source/Runtime/D3D12RHI/Private/D3D12Submission.cpp b/Engine/Source/Runtime/D3D12RHI/Private/D3D12Submission.cpp
--- a/Engine/Source/Runtime/D3D12RHI/Private/D3D12Submission.cpp
+++ b/Engine/Source/Runtime/D3D12RHI/Private/D3D12Submission.cpp
@@ -72,6 +72,7 @@ static FAutoConsoleVariableRef CVarD3D12SubmissionMaxExecuteBatchSizeAsync(
 	TEXT("The valid range is 1 to INT_MAX inclusive. Values less than 1 will be clamped to 1"),
 	ECVF_RenderThreadSafe);
 
+static HRESULT GGPUD3DResult{};
 static std::atomic<int> GGPUCrashDetected = false;
 
 class FD3D12Thread final : private FRunnable
@@ -968,13 +969,14 @@ void FD3D12DynamicRHI::ProcessInterruptQueueUntil(FGraphEvent* GraphEvent)
 	}
 }
 
-void FD3D12DynamicRHI::ProcessInterruptQueueOnGPUCrash()
+void FD3D12DynamicRHI::ProcessInterruptQueueOnGPUCrash(HRESULT D3DResult)
 {
 	// This function will not return.
 
 	// We know this function was called due to a GPU crash, so let the thread know.
-	GGPUCrashDetected.store(true, std::memory_order_release);
-	
+	GGPUD3DResult = D3DResult;
+	GGPUCrashDetected.store(true);
+
 	if (InterruptThread)
 	{
 		// Since we have an interrupt thread, allow it to process the GPU crash.
@@ -1066,7 +1068,7 @@ FD3D12DynamicRHI::FProcessResult FD3D12DynamicRHI::ProcessInterruptQueue()
 		HRESULT DeviceRemovedReason = Queue.Device->GetDevice()->GetDeviceRemovedReason();
 		if (DeviceRemovedReason != S_OK)
 		{
-			TerminateOnGPUCrash();
+			TerminateOnGPUCrash(NullOpt);
 		}
 	};
 
@@ -1240,11 +1242,11 @@ FD3D12DynamicRHI::FProcessResult FD3D12DynamicRHI::ProcessInterruptQueue()
 		CheckForDeviceRemoved(CurrentQueue);
 	});
 
-	if (GGPUCrashDetected.load(std::memory_order_relaxed))
+	if (GGPUCrashDetected.load())
 	{
 		// If this was set by ProcessInterruptQueueOnGPUCrash, we know a crash was detected, so process it immediately. We can't always rely on
 		// queue processing to catch it, as GetDeviceRemovedReason sometimes returns S_OK despite an earlier API call having reported a lost device.
-		TerminateOnGPUCrash();
+		TerminateOnGPUCrash(GGPUD3DResult);
 	}
 
 	return Result;
diff --git a/Engine/Source/Runtime/D3D12RHI/Private/D3D12Util.cpp b/Engine/Source/Runtime/D3D12RHI/Private/D3D12Util.cpp
--- a/Engine/Source/Runtime/D3D12RHI/Private/D3D12Util.cpp
+++ b/Engine/Source/Runtime/D3D12RHI/Private/D3D12Util.cpp
@@ -979,8 +979,29 @@ void FD3D12DynamicRHI::OutputGPUCrashReport(FTextBuilder& ErrorMessage)
 #endif
 }
 
-void FD3D12DynamicRHI::TerminateOnGPUCrash()
+void FD3D12DynamicRHI::TerminateOnGPUCrash(TOptional<HRESULT> D3DResult)
 {
+	// Include the D3D HRESULT that alerted the application of the GPU crash and the device removal reason, if present.
+	// In some cases, such as when the runtime signals the device is lost through fence values, there is no HRESULT
+	// but there might be a device removal reason.
+	if (D3DResult.IsSet())
+	{
+		FGenericCrashContext::SetEngineData(TEXT("GPUCrash.D3DResult"), FString::FromInt(*D3DResult));
+	}
+
+	TConstArrayView<FD3D12Device*> Devices = GetAdapter().GetDevices();
+	for (FD3D12Device* Device : Devices)
+	{
+		HRESULT Reason = Device->GetDevice()->GetDeviceRemovedReason();
+		if (Reason != S_OK)
+		{
+			FGenericCrashContext::SetEngineData(TEXT("GPUCrash.D3DDeviceRemovedReason"), FString::FromInt(Reason));
+
+			// Just include one device's removal reason in the crash report payload.
+			break;
+		}
+	}
+
 #if WITH_NVAPI_RAYTRACING_VALIDATION
 	FlushRaytracingValidationMessages();
 #endif
@@ -1064,7 +1085,7 @@ void FD3D12DynamicRHI::HandleFailedD3D12Result(HRESULT D3DResult, ID3D12Device*
 			}
 
 			// This function does not return.
-			TerminateOnGPUCrash();
+			TerminateOnGPUCrash(D3DResult);
 		}
 		else
 		{
@@ -1074,7 +1095,7 @@ void FD3D12DynamicRHI::HandleFailedD3D12Result(HRESULT D3DResult, ID3D12Device*
 
 			// The interrupt thread must be the one to handle DXGI_ERROR_DEVICE_REMOVED etc.
 			// This function does not return.
-			ProcessInterruptQueueOnGPUCrash();
+			ProcessInterruptQueueOnGPUCrash(D3DResult);
 		}
 	}
 	else
diff --git a/Engine/Source/Runtime/Windows/D3D11RHI/Private/D3D11Util.cpp b/Engine/Source/Runtime/Windows/D3D11RHI/Private/D3D11Util.cpp
--- a/Engine/Source/Runtime/Windows/D3D11RHI/Private/D3D11Util.cpp
+++ b/Engine/Source/Runtime/Windows/D3D11RHI/Private/D3D11Util.cpp
@@ -131,6 +131,16 @@ static void TerminateOnDeviceRemoved(HRESULT D3DResult, ID3D11Device* Direct3DDe
 {
 	if (D3DResult == DXGI_ERROR_DEVICE_REMOVED)
 	{
+		FGenericCrashContext::SetEngineData(TEXT("GPUCrash.D3DResult"), FString::FromInt(D3DResult));
+		if (Direct3DDevice)
+		{
+			const HRESULT hRes = Direct3DDevice->GetDeviceRemovedReason();
+			if (hRes != S_OK)
+			{
+				FGenericCrashContext::SetEngineData(TEXT("GPUCrash.D3DDeviceRemovedReason"), FString::FromInt(hRes));
+			}
+		}
+
 	#if NV_AFTERMATH
 		TArray<UE::RHICore::Nvidia::Aftermath::FCrashResult> AftermathResults;
 		UE::RHICore::Nvidia::Aftermath::OnGPUCrash(AftermathResults);
@@ -145,7 +155,7 @@ static void TerminateOnDeviceRemoved(HRESULT D3DResult, ID3D11Device* Direct3DDe
 		// Report the GPU crash which will raise the exception
 		ReportGPUCrash(TEXT("GPU Crash dump Triggered"), nullptr);
 
-		GIsGPUCrashed = true;		
+		GIsGPUCrashed = true;
 		if (Direct3DDevice)
 		{
 			HRESULT hRes = Direct3DDevice->GetDeviceRemovedReason();
@@ -160,7 +170,7 @@ static void TerminateOnDeviceRemoved(HRESULT D3DResult, ID3D11Device* Direct3DDe
 			case DXGI_ERROR_INVALID_CALL:			Reason = TEXT("INVALID_CALL"); break;
 			case S_OK:								Reason = TEXT("S_OK"); break;
 			}
-			
+
 			// We currently don't support removed devices because FTexture2DResource can't recreate its RHI resources from scratch.
 			// We would also need to recreate the viewport swap chains from scratch.			
 			UE_LOG(LogD3D11RHI, Fatal, TEXT("Unreal Engine is exiting due to D3D device being lost. (Error: 0x%X - '%s')"), hRes, Reason);
