Make sure D3D12 Descriptor Range resources are resident before use.

#Synced-CL 51674122

diff --git a/Engine/Source/Runtime/D3D12RHI/Private/D3D12Commands.cpp b/Engine/Source/Runtime/D3D12RHI/Private/D3D12Commands.cpp
--- a/Engine/Source/Runtime/D3D12RHI/Private/D3D12Commands.cpp
+++ b/Engine/Source/Runtime/D3D12RHI/Private/D3D12Commands.cpp
@@ -601,6 +601,10 @@ struct FD3D12ResourceBinder
 			{
 				FRHIDescriptorRangeShaderData ShaderData = D3D12DescriptorRange->GetShaderData();
 				SetParameter(0, Offset, sizeof(ShaderData), reinterpret_cast<const uint8*>(&ShaderData));
+
+				// Make sure the referenced resources are resident before use.
+				Context.StateCache.QueueBindlessSRVs(Frequency, D3D12DescriptorRange->AllSrvs);
+				Context.StateCache.QueueBindlessUAVs(Frequency, D3D12DescriptorRange->AllUavs);
 			}
 		}
 	}
diff --git a/Engine/Source/Runtime/D3D12RHI/Private/D3D12DescriptorRange.cpp b/Engine/Source/Runtime/D3D12RHI/Private/D3D12DescriptorRange.cpp
--- a/Engine/Source/Runtime/D3D12RHI/Private/D3D12DescriptorRange.cpp
+++ b/Engine/Source/Runtime/D3D12RHI/Private/D3D12DescriptorRange.cpp
@@ -51,6 +51,7 @@ FD3D12DescriptorRange::~FD3D12DescriptorRange()
 {
 	if (Allocation.Count > 0)
 	{
+		// Deferred so in-flight command lists can finish referencing these descriptors.
 		FD3D12DynamicRHI::GetD3DRHI()->DeferredDelete(Allocation, GetDesc().TypeMask, GetParentDevice());
 	}
 }
@@ -72,19 +73,28 @@ void FD3D12DescriptorRange::InitializeDescriptors(TConstArrayView<FRHIDescriptor
 
 	FD3D12BindlessDescriptorManager& Manager = Device->GetBindlessDescriptorManager();
 
-	for (int32 Index = 0; Index < InResources.Num(); Index++)
+	const int32 NumResources = FMath::Min(InResources.Num(), (int32)GetDesc().Count);
+	ResourceTypes.SetNum(NumResources);
+
+	for (int32 Index = 0; Index < NumResources; Index++)
 	{
-		if (Index < (int32)GetDesc().Count)
-		{
-			const FRHIDescriptorResource& Resource = InResources[Index];
+		const FRHIDescriptorResource& Resource = InResources[Index];
 
-			const FD3D12View* View = GetViewFromDescriptorResource(Resource, GpuIndex);
-			check(EnumHasAllFlags(DescriptorTypeMask, View->GetBindlessHandle().GetTypeMask()));
+		const FD3D12View* View = GetViewFromDescriptorResource(Resource, GpuIndex);
 
-			const FRHIDescriptorHandle UpdateHandle(View->GetBindlessHandle().GetType(), Allocation.StartIndex + Index);
-			Manager.InitializeDescriptor(UpdateHandle, View);
-		}
+		checkf(EnumHasAllFlags(DescriptorTypeMask, View->GetBindlessHandle().GetTypeMask()),
+			TEXT("View descriptor type %s is not compatible with descriptor range type mask %s."),
+			GetDescriptorTypeString(View->GetBindlessHandle().GetType()),
+			*GetDescriptorMaskString(DescriptorTypeMask));
+
+		const FRHIDescriptorHandle UpdateHandle(View->GetBindlessHandle().GetType(), Allocation.StartIndex + Index);
+		Manager.InitializeDescriptor(UpdateHandle, View);
+
+		ResourceTypes[Index] = Resource.Type;
 	}
+
+	// Rebuild cached views so residency is correct at next bind.
+	PopulateResidencyViews();
 }
 
 void FD3D12DescriptorRange::UpdateResources(const FRHIContextArray& Contexts, uint32 FirstResourceIndex, TConstArrayView<FRHIDescriptorResource> NewResources)
@@ -105,10 +115,51 @@ void FD3D12DescriptorRange::UpdateResources(const FRHIContextArray& Contexts, ui
 			UpdateResource(FirstResourceIndex + Index, Resource);
 
 			const FD3D12View* View = GetViewFromDescriptorResource(Resource, GpuIndex);
-			check(EnumHasAllFlags(DescriptorTypeMask, View->GetBindlessHandle().GetTypeMask()));
+
+			checkf(EnumHasAllFlags(DescriptorTypeMask, View->GetBindlessHandle().GetTypeMask()),
+				TEXT("View descriptor type %s is not compatible with descriptor range type mask %s."),
+				GetDescriptorTypeString(View->GetBindlessHandle().GetType()),
+				*GetDescriptorMaskString(DescriptorTypeMask));
 
 			const FRHIDescriptorHandle UpdateHandle(View->GetBindlessHandle().GetType(), Allocation.StartIndex + InternalIndex);
 			Manager.UpdateDescriptor(Contexts, UpdateHandle, View);
+
+			ResourceTypes[InternalIndex] = Resource.Type;
+		}
+	}
+
+	// Rebuild cached views so residency is correct at next bind.
+	PopulateResidencyViews();
+}
+
+void FD3D12DescriptorRange::PopulateResidencyViews()
+{
+	AllSrvs.Reset();
+	AllUavs.Reset();
+
+	const uint32 GpuIndex = GetParentDevice()->GetGPUIndex();
+
+	for (int32 Index = 0; Index < Resources.Num(); Index++)
+	{
+		FRHIResource* Resource = Resources[Index].GetReference();
+		if (!Resource)
+		{
+			continue;
+		}
+
+		// Reconstruct the FRHIDescriptorResource from the stored type so that we can get the correct view from the resource.
+		const FRHIDescriptorResource DescriptorResource(ResourceTypes[Index], Resource);
+
+		if (FD3D12View* View = GetViewFromDescriptorResource(DescriptorResource, GpuIndex))
+		{
+			if (ResourceTypes[Index] == FRHIDescriptorResource::EType::UnorderedAccessView)
+			{
+				AllUavs.Emplace(static_cast<FD3D12UnorderedAccessView*>(View));
+			}
+			else
+			{
+				AllSrvs.Emplace(static_cast<FD3D12ShaderResourceView*>(View));
+			}
 		}
 	}
 }
@@ -147,6 +198,7 @@ void FD3D12DynamicRHI::RHIUpdateDescriptorRange(FRHICommandListBase& RHICmdList,
 
 	if (FirstResourceIndex < DescriptorRange->GetDesc().Count)
 	{
+		// Clamp the incoming resources to not exceed the range's capacity.
 		TConstArrayView<FRHIDescriptorResource> NewResources = InNewResources.Left(DescriptorRange->GetDesc().Count - FirstResourceIndex);
 		TConstArrayView<FRHIDescriptorResource> NewResourcesCmd = RHICmdList.IsTopOfPipe() ? RHICmdList.AllocArray(NewResources) : NewResources;
 
diff --git a/Engine/Source/Runtime/D3D12RHI/Private/D3D12DescriptorRange.h b/Engine/Source/Runtime/D3D12RHI/Private/D3D12DescriptorRange.h
--- a/Engine/Source/Runtime/D3D12RHI/Private/D3D12DescriptorRange.h
+++ b/Engine/Source/Runtime/D3D12RHI/Private/D3D12DescriptorRange.h
@@ -6,9 +6,13 @@
 #include "RHIDescriptorAllocator.h"
 #include "D3D12RHICommon.h"
 
+class FD3D12ShaderResourceView;
+class FD3D12UnorderedAccessView;
+
 class FD3D12DescriptorRange : public FRHIDescriptorRange, public FD3D12DeviceChild, public FD3D12LinkedAdapterObject<FD3D12DescriptorRange>
 {
 public:
+	// Constructs an invalid/empty descriptor range. Used as a fallback when allocation fails.
 	FD3D12DescriptorRange(FD3D12Device* InParent);
 	FD3D12DescriptorRange(FD3D12Device* InParent, const FRHIDescriptorRangeDesc& InDesc, TConstArrayView<FRHIDescriptorResource> InResources, FRHIDescriptorAllocation Allocation);
 	~FD3D12DescriptorRange();
@@ -18,7 +22,19 @@ public:
 	void InitializeDescriptors(TConstArrayView<FRHIDescriptorResource> InResources);
 	void UpdateResources(const FRHIContextArray& Contexts, uint32 FirstResourceIndex, TConstArrayView<FRHIDescriptorResource> NewResources);
 
+	// Rebuilds the cached AllSrvs/AllUavs arrays from the current resource list.
+	void PopulateResidencyViews();
+
 	FRHIDescriptorAllocation Allocation;
+
+	// Stores the descriptor type (Texture, SRV, UAV) for each slot. Used by PopulateResidencyViews
+	// to extract the correct D3D12 view type from the base class Resources array.
+	TArray<FRHIDescriptorResource::EType> ResourceTypes;
+
+	// Cached D3D12 views for all resources in the range. Rebuilt on create/update via PopulateResidencyViews().
+	// Queued for residency (via QueueBindlessSRVs/UAVs) when the descriptor range is bound to a shader.
+	TArray<FD3D12ShaderResourceView*> AllSrvs;
+	TArray<FD3D12UnorderedAccessView*> AllUavs;
 };
 
 template<>
