[D3D12RHI] Fix FD3D12SyncPoint for possibly dangling TCHAR* DebugName using FName instead

#Synced-CL 51754672

diff --git a/Engine/Source/Runtime/D3D12RHI/Private/D3D12CommandContext.cpp b/Engine/Source/Runtime/D3D12RHI/Private/D3D12CommandContext.cpp
--- a/Engine/Source/Runtime/D3D12RHI/Private/D3D12CommandContext.cpp
+++ b/Engine/Source/Runtime/D3D12RHI/Private/D3D12CommandContext.cpp
@@ -537,7 +537,7 @@ void FD3D12QueryAllocator::CloseAndReset(TMap<TRefCountPtr<FD3D12QueryHeap>, TAr
 
 FD3D12CopyScope::FD3D12CopyScope(FD3D12Device* Device, ED3D12SyncPointType SyncPointType, FD3D12SyncPointRef const& WaitSyncPoint)
 	: Device(Device)
-	, SyncPoint(FD3D12SyncPoint::Create(SyncPointType, TEXT("CopyScope")))
+	, SyncPoint(FD3D12SyncPoint::Create(SyncPointType, UE_D3D12_STATIC_FNAME("CopyScope")))
 	, Context(*Device->ObtainContextCopy())
 {
 	if (WaitSyncPoint)
@@ -612,7 +612,7 @@ void FD3D12ContextCommon::FlushCommands(ED3D12FlushFlags FlushFlags)
 
 	if (EnumHasAnyFlags(FlushFlags, ED3D12FlushFlags::WaitForCompletion))
 	{
-		SyncPoint = FD3D12SyncPoint::Create(ED3D12SyncPointType::GPUAndCPU, TEXT("WaitForCompletion"));
+		SyncPoint = FD3D12SyncPoint::Create(ED3D12SyncPointType::GPUAndCPU, UE_D3D12_STATIC_FNAME("WaitForCompletion"));
 		SignalSyncPoint(SyncPoint);
 	}
 
@@ -812,7 +812,7 @@ void FD3D12CommandContextRedirector::RHITransferResources(TConstArrayView<FTrans
 		TStaticArray<FD3D12SyncPointRef, MAX_NUM_GPUS> SyncPoints;
 		for (uint32 GPUIndex : SignalMask)
 		{
-			SyncPoints[GPUIndex] = FD3D12SyncPoint::Create(ED3D12SyncPointType::GPUOnly, TEXT("TransferResources"));
+			SyncPoints[GPUIndex] = FD3D12SyncPoint::Create(ED3D12SyncPointType::GPUOnly, UE_D3D12_STATIC_FNAME("TransferResources"));
 			PhysicalContexts[GPUIndex]->SignalSyncPoint(SyncPoints[GPUIndex]);
 		}
 
@@ -1027,7 +1027,7 @@ void FD3D12CommandContextRedirector::RHITransferResourceSignal(TConstArrayView<F
 	uint32 FenceIndex = 0;
 	for (uint32 SrcGPUIndex : SrcGPUMask)
 	{
-		FD3D12SyncPointRef SyncPoint = FD3D12SyncPoint::Create(ED3D12SyncPointType::GPUOnly, TEXT("TransferReourceSignal"));
+		FD3D12SyncPointRef SyncPoint = FD3D12SyncPoint::Create(ED3D12SyncPointType::GPUOnly, UE_D3D12_STATIC_FNAME("TransferResourceSignal"));
 		SyncPoint->AddRef();
 
 		PhysicalContexts[SrcGPUIndex]->SignalSyncPoint(SyncPoint);
@@ -1146,7 +1146,7 @@ void FD3D12CommandContextRedirector::RHICrossGPUTransfer(TConstArrayView<FTransf
 	}
 
 	// Post-copy synchronization
-	FD3D12SyncPointRef SyncPoint = FD3D12SyncPoint::Create(ED3D12SyncPointType::GPUOnly, TEXT("RHICrossGPUTransfer"));
+	FD3D12SyncPointRef SyncPoint = FD3D12SyncPoint::Create(ED3D12SyncPointType::GPUOnly, UE_D3D12_STATIC_FNAME("RHICrossGPUTransfer"));
 	PhysicalContexts[Params[0].SrcGPUIndex]->SignalSyncPoint(SyncPoint);
 
 	for (FCrossGPUTransferFence* PostTransferSyncPoint : PostTransfer)
@@ -1161,7 +1161,7 @@ void FD3D12CommandContextRedirector::RHICrossGPUTransferSignal(TConstArrayView<F
 {
 	for (FCrossGPUTransferFence* TransferSyncPoint : PreTransfer)
 	{
-		FD3D12SyncPointRef SyncPoint = FD3D12SyncPoint::Create(ED3D12SyncPointType::GPUOnly, TEXT("RHICrossGPUTransferSignal"));
+		FD3D12SyncPointRef SyncPoint = FD3D12SyncPoint::Create(ED3D12SyncPointType::GPUOnly, UE_D3D12_STATIC_FNAME("RHICrossGPUTransferSignal"));
 		SyncPoint->AddRef();
 
 		PhysicalContexts[TransferSyncPoint->SignalGPUIndex]->SignalSyncPoint(SyncPoint);
diff --git a/Engine/Source/Runtime/D3D12RHI/Private/D3D12CommandContext.h b/Engine/Source/Runtime/D3D12RHI/Private/D3D12CommandContext.h
--- a/Engine/Source/Runtime/D3D12RHI/Private/D3D12CommandContext.h
+++ b/Engine/Source/Runtime/D3D12RHI/Private/D3D12CommandContext.h
@@ -282,7 +282,7 @@ public:
 	{
 		if (!ContextSyncPoint)
 		{
-			ContextSyncPoint = FD3D12SyncPoint::Create(ED3D12SyncPointType::GPUAndCPU, TEXT("ContextSyncPoint"));
+			ContextSyncPoint = FD3D12SyncPoint::Create(ED3D12SyncPointType::GPUAndCPU, UE_D3D12_STATIC_FNAME("ContextSyncPoint"));
 			BatchedSyncPoints.ToSignal.Add(ContextSyncPoint);
 		}
 
diff --git a/Engine/Source/Runtime/D3D12RHI/Private/D3D12Device.cpp b/Engine/Source/Runtime/D3D12RHI/Private/D3D12Device.cpp
--- a/Engine/Source/Runtime/D3D12RHI/Private/D3D12Device.cpp
+++ b/Engine/Source/Runtime/D3D12RHI/Private/D3D12Device.cpp
@@ -553,7 +553,7 @@ void FD3D12Device::BlockUntilIdle()
 
 	for (uint32 QueueTypeIndex = 0; QueueTypeIndex < (uint32)ED3D12QueueType::Count; ++QueueTypeIndex)
 	{
-		FD3D12SyncPointRef SyncPoint = FD3D12SyncPoint::Create(ED3D12SyncPointType::GPUAndCPU, TEXT("BlockUntilGPUIdle"));
+		FD3D12SyncPointRef SyncPoint = FD3D12SyncPoint::Create(ED3D12SyncPointType::GPUAndCPU, UE_D3D12_STATIC_FNAME("BlockUntilGPUIdle"));
 
 		FD3D12Payload* Payload = new FD3D12Payload(GetQueue((ED3D12QueueType)QueueTypeIndex));
 		Payload->SyncPointsToSignal.Add(SyncPoint);
diff --git a/Engine/Source/Runtime/D3D12RHI/Private/D3D12DirectCommandListManager.cpp b/Engine/Source/Runtime/D3D12RHI/Private/D3D12DirectCommandListManager.cpp
--- a/Engine/Source/Runtime/D3D12RHI/Private/D3D12DirectCommandListManager.cpp
+++ b/Engine/Source/Runtime/D3D12RHI/Private/D3D12DirectCommandListManager.cpp
@@ -62,13 +62,7 @@ void FD3D12DynamicRHI::RHIWriteGPUFence_TopOfPipe(FRHICommandListBase& RHICmdLis
 	{
 		checkf(Fence->SyncPoints[GPUIndex] == nullptr, TEXT("The fence for the current GPU node has already been issued."));
 
-		Fence->SyncPoints[GPUIndex] = FD3D12SyncPoint::Create(ED3D12SyncPointType::GPUAndCPU,
-#if RHI_USE_SYNC_POINT_DEBUG_NAME
-		*FenceRHI->GetFName().ToString()
-#else 
-		TEXT("")
-#endif
-		);
+		Fence->SyncPoints[GPUIndex] = FD3D12SyncPoint::Create(ED3D12SyncPointType::GPUAndCPU, FenceRHI->GetFName());
 	}
 
 	Fence->NumPendingWriteCommands.Increment();
diff --git a/Engine/Source/Runtime/D3D12RHI/Private/D3D12EnhancedBarriers.cpp b/Engine/Source/Runtime/D3D12RHI/Private/D3D12EnhancedBarriers.cpp
--- a/Engine/Source/Runtime/D3D12RHI/Private/D3D12EnhancedBarriers.cpp
+++ b/Engine/Source/Runtime/D3D12RHI/Private/D3D12EnhancedBarriers.cpp
@@ -2346,7 +2346,7 @@ void FD3D12EnhancedBarriersForAdapterImpl::CreateTransition(
 			TRHIPipelineArray<FD3D12SyncPointRef>& DeviceSyncPoints = Data->SyncPoints.Emplace_GetRef();
 			for (ERHIPipeline Pipeline : MakeFlagsRange(CreateInfo.SrcPipelines))
 			{
-				DeviceSyncPoints[Pipeline] = FD3D12SyncPoint::Create(ED3D12SyncPointType::GPUOnly, TEXT("Transition"));
+				DeviceSyncPoints[Pipeline] = FD3D12SyncPoint::Create(ED3D12SyncPointType::GPUOnly, UE_D3D12_STATIC_FNAME("Transition"));
 			}
 		}
 	}
diff --git a/Engine/Source/Runtime/D3D12RHI/Private/D3D12LegacyBarriers.cpp b/Engine/Source/Runtime/D3D12RHI/Private/D3D12LegacyBarriers.cpp
--- a/Engine/Source/Runtime/D3D12RHI/Private/D3D12LegacyBarriers.cpp
+++ b/Engine/Source/Runtime/D3D12RHI/Private/D3D12LegacyBarriers.cpp
@@ -413,7 +413,7 @@ void FD3D12LegacyBarriersForAdapterImpl::CreateTransition(
 			TRHIPipelineArray<FD3D12SyncPointRef>& DeviceSyncPoints = Data->SyncPoints.Emplace_GetRef();
 			for (ERHIPipeline Pipeline : MakeFlagsRange(CreateInfo.SrcPipelines))
 			{
-				DeviceSyncPoints[Pipeline] = FD3D12SyncPoint::Create(ED3D12SyncPointType::GPUOnly, TEXT("Transition"));
+				DeviceSyncPoints[Pipeline] = FD3D12SyncPoint::Create(ED3D12SyncPointType::GPUOnly, UE_D3D12_STATIC_FNAME("Transition"));
 			}
 		}
 	}
diff --git a/Engine/Source/Runtime/D3D12RHI/Private/D3D12Queue.h b/Engine/Source/Runtime/D3D12RHI/Private/D3D12Queue.h
--- a/Engine/Source/Runtime/D3D12RHI/Private/D3D12Queue.h
+++ b/Engine/Source/Runtime/D3D12RHI/Private/D3D12Queue.h
@@ -39,6 +39,17 @@ inline const TCHAR* GetD3DCommandQueueTypeName(ED3D12QueueType QueueType)
 	}
 }
 
+inline FName GetD3DCommandQueueTypeFName(ED3D12QueueType QueueType)
+{
+	switch (QueueType)
+	{
+	default: checkNoEntry(); // fallthrough
+	case ED3D12QueueType::Direct: { static const FName Name(TEXT("3D"));      return Name; }
+	case ED3D12QueueType::Async:  { static const FName Name(TEXT("Compute")); return Name; }
+	case ED3D12QueueType::Copy:   { static const FName Name(TEXT("Copy"));    return Name; }
+	}
+}
+
 inline D3D12_COMMAND_LIST_TYPE GetD3DCommandListType(ED3D12QueueType QueueType)
 {
 	switch (QueueType)
diff --git a/Engine/Source/Runtime/D3D12RHI/Private/D3D12RHI.cpp b/Engine/Source/Runtime/D3D12RHI/Private/D3D12RHI.cpp
--- a/Engine/Source/Runtime/D3D12RHI/Private/D3D12RHI.cpp
+++ b/Engine/Source/Runtime/D3D12RHI/Private/D3D12RHI.cpp
@@ -423,7 +423,7 @@ void FD3D12DynamicRHI::EnqueueEndOfPipeTask(TUniqueFunction<void()> TaskFunc, TU
 	{
 		FD3D12Payload* Payload = new FD3D12Payload(Queue);
 
-		FD3D12SyncPointRef SyncPoint = FD3D12SyncPoint::Create(ED3D12SyncPointType::GPUAndCPU, GetD3DCommandQueueTypeName(Queue.QueueType) );
+		FD3D12SyncPointRef SyncPoint = FD3D12SyncPoint::Create(ED3D12SyncPointType::GPUAndCPU, GetD3DCommandQueueTypeFName(Queue.QueueType));
 		Payload->SyncPointsToSignal.Emplace(SyncPoint);
 		Prereqs.Add(SyncPoint->GetGraphEvent());
 
diff --git a/Engine/Source/Runtime/D3D12RHI/Private/D3D12Submission.cpp b/Engine/Source/Runtime/D3D12RHI/Private/D3D12Submission.cpp
--- a/Engine/Source/Runtime/D3D12RHI/Private/D3D12Submission.cpp
+++ b/Engine/Source/Runtime/D3D12RHI/Private/D3D12Submission.cpp
@@ -916,7 +916,7 @@ void FD3D12SyncPoint::Wait() const
 		// Block the calling thread until the graph event is signaled by the interrupt thread.
 #if CPUPROFILERTRACE_ENABLED
 		UE_TRACE_LOG_SCOPED_T(Cpu, SyncPoint_Wait, CpuChannel)
-			<< SyncPoint_Wait.WaitingForName(GetDebugName());
+			<< SyncPoint_Wait.WaitingForName(*GetDebugName().ToString());
 #endif // CPUPROFILERTRACE_ENABLED
 
 		FD3D12DynamicRHI::GetD3DRHI()->ProcessInterruptQueueUntil(GraphEvent);
diff --git a/Engine/Source/Runtime/D3D12RHI/Private/D3D12Submission.h b/Engine/Source/Runtime/D3D12RHI/Private/D3D12Submission.h
--- a/Engine/Source/Runtime/D3D12RHI/Private/D3D12Submission.h
+++ b/Engine/Source/Runtime/D3D12RHI/Private/D3D12Submission.h
@@ -27,6 +27,9 @@ class FD3D12Viewport;
 class FD3D12SyncPoint;
 using FD3D12SyncPointRef = TRefCountPtr<FD3D12SyncPoint>;
 
+// Helper macro to create a static FName from a string literal, avoiding repeated FName construction.
+#define UE_D3D12_STATIC_FNAME(String) (([]() -> const FName& { static const FName Name(TEXT(String)); return Name; })())
+
 enum class ED3D12SyncPointType
 {
 	// Sync points of this type do not include an FGraphEvent, so cannot
@@ -87,7 +90,7 @@ class FD3D12SyncPoint final : public FRefCountedObject
 	TOptional<FD3D12ResolvedFence> ResolvedFence;
 	FGraphEventRef GraphEvent;
 
-	FD3D12SyncPoint(ED3D12SyncPointType Type, const TCHAR* InDebugName)
+	FD3D12SyncPoint(ED3D12SyncPointType Type, FName InDebugName)
 #if RHI_USE_SYNC_POINT_DEBUG_NAME
 		:DebugName(InDebugName)
 #endif
@@ -99,7 +102,7 @@ class FD3D12SyncPoint final : public FRefCountedObject
 	}
 
 public:
-	static FD3D12SyncPointRef Create(ED3D12SyncPointType Type, const TCHAR* InDebugName)
+	static FD3D12SyncPointRef Create(ED3D12SyncPointType Type, FName InDebugName)
 	{
 		LLM_SCOPE_BYNAME(TEXT("RHIMisc/CreateSyncPoint"));
 		return new FD3D12SyncPoint(Type, InDebugName);
@@ -143,19 +146,19 @@ public:
 		MemoryPool.Push(Pointer);
 	}
 	 
-	inline const TCHAR* GetDebugName() const
+	inline FName GetDebugName() const
 	{
 #if RHI_USE_SYNC_POINT_DEBUG_NAME
 		return DebugName;
 #else
-		return TEXT("SyncPointBusyWait");
+		return UE_D3D12_STATIC_FNAME("SyncPointBusyWait");
 #endif
 	}
 
 private:
 
 #if RHI_USE_SYNC_POINT_DEBUG_NAME
-	const TCHAR* DebugName = nullptr;
+	FName DebugName;
 #endif
 
 };
diff --git a/Engine/Source/Runtime/D3D12RHI/Private/D3D12Viewport.cpp b/Engine/Source/Runtime/D3D12RHI/Private/D3D12Viewport.cpp
--- a/Engine/Source/Runtime/D3D12RHI/Private/D3D12Viewport.cpp
+++ b/Engine/Source/Runtime/D3D12RHI/Private/D3D12Viewport.cpp
@@ -806,7 +806,7 @@ void FD3D12Viewport::IssueFrameEvent()
 	{
 		FD3D12CommandContext& Context = Device->GetDefaultCommandContext();
 
-		FD3D12SyncPointRef SyncPoint = FD3D12SyncPoint::Create(ED3D12SyncPointType::GPUAndCPU, TEXT("WaitPreviousFrameEvent"));
+		FD3D12SyncPointRef SyncPoint = FD3D12SyncPoint::Create(ED3D12SyncPointType::GPUAndCPU, UE_D3D12_STATIC_FNAME("WaitPreviousFrameEvent"));
 
 		Context.SignalSyncPoint(SyncPoint);
 		Context.Finalize(Payloads);
diff --git a/Engine/Source/Runtime/RHI/Public/RHIStrings.h b/Engine/Source/Runtime/RHI/Public/RHIStrings.h
--- a/Engine/Source/Runtime/RHI/Public/RHIStrings.h
+++ b/Engine/Source/Runtime/RHI/Public/RHIStrings.h
@@ -113,13 +113,13 @@ inline const TCHAR* GetShaderFrequencyString(EShaderFrequency Frequency, bool bI
 	return String;
 }
 
-inline const TCHAR* GetRenderQueryTypeName(ERenderQueryType QueryType)
+inline FName GetRenderQueryTypeName(ERenderQueryType QueryType)
 {
 	switch (QueryType)
 	{
 	default: checkNoEntry(); // fallthrough
-	case RQT_Undefined:		return TEXT("Undefined");
-	case RQT_Occlusion:		return TEXT("GPUOcclusion");
-	case RQT_AbsoluteTime:	return TEXT("GPUAbsoluteTime");
+	case RQT_Undefined:		{ static const FName Name(TEXT("Undefined"));      return Name; }
+	case RQT_Occlusion:		{ static const FName Name(TEXT("GPUOcclusion"));    return Name; }
+	case RQT_AbsoluteTime:	{ static const FName Name(TEXT("GPUAbsoluteTime")); return Name; }
 	}
 }
