Quiet D3DDebug Warnings About Duplicate Subresource Transitions

The D3D12 debug layer will complain if there are multiple barriers for the same (sub)resource in a single barrier batch. While it's not ideal that we do this, some of it isn't, practically, avoidable.  This teaches the legacy barriers to merge a new barrier with an existing barrier if possible, by replacing its StateAfter with the following barrier's StateAfter, which reduces the complaints from the debug layer by some amount. It then silences the remaining warnings since they aren't currently actionable.

#Synced-CL 51818073

diff --git a/Engine/Source/Runtime/D3D12RHI/Private/D3D12Adapter.cpp b/Engine/Source/Runtime/D3D12RHI/Private/D3D12Adapter.cpp
--- a/Engine/Source/Runtime/D3D12RHI/Private/D3D12Adapter.cpp
+++ b/Engine/Source/Runtime/D3D12RHI/Private/D3D12Adapter.cpp
@@ -918,6 +918,14 @@ void FD3D12Adapter::CreateD3DInfoQueue()
 			// See: UE-222685 to remove again from Deny list after Agility SDK upgrade
 			D3D12_MESSAGE_ID_HEAP_ADDRESS_RANGE_HAS_NO_RESOURCE,
 
+#if D3D12RHI_SUPPORTS_LEGACY_BARRIERS
+			// This happens a lot because when we perform copies we'll transition a resource into the COPY_DEST or COPY_SRC state temporarily and then
+			// transition it back into whatever state it was in before. Often there's no work done between that transition out of COPY_[DEST|SRC] and the next
+			// transition. We already try to minimize this by identifying back-to-back barriers that can be combined, or nullify each other, but it still
+			// happens when there are other barriers in-between. Short of tracking barriers by resource in the batcher, I don't see a way to avoid it given 
+			// current architecture. The warning is quite noisy so turning it off should help us recognize more actionable and important warnings.
+			D3D12_MESSAGE_ID_RESOURCE_BARRIER_DUPLICATE_SUBRESOURCE_TRANSITIONS,
+#endif
 		};
 
 #if PLATFORM_DESKTOP
diff --git a/Engine/Source/Runtime/D3D12RHI/Private/D3D12LegacyBarriers.cpp b/Engine/Source/Runtime/D3D12RHI/Private/D3D12LegacyBarriers.cpp
--- a/Engine/Source/Runtime/D3D12RHI/Private/D3D12LegacyBarriers.cpp
+++ b/Engine/Source/Runtime/D3D12RHI/Private/D3D12LegacyBarriers.cpp
@@ -829,21 +829,32 @@ int32 FD3D12LegacyBarriersBatcher::AddTransition(
 
 	if (Barriers.Num())
 	{
-		// Check if we are simply reverting the last transition. In that case, we can just remove both transitions.
-		// This happens fairly frequently due to resource pooling since different RHI buffers can point to the same underlying D3D buffer.
-		// Instead of ping-ponging that underlying resource between COPY_DEST and GENERIC_READ, several copies can happen without a ResourceBarrier() in between.
-		// Doing this check also eliminates a D3D debug layer warning about multiple transitions of the same subresource.
-		const FD3D12ResourceBarrier& Last = Barriers.Last();
-
+		// Check if the new transition is for the same (sub)resource as the prior barrier, which also must be a transition.
+		FD3D12ResourceBarrier& Last = Barriers.Last();
 		if (Last.Type == D3D12_RESOURCE_BARRIER_TYPE_TRANSITION
 			&& pResource->GetResource() == Last.Transition.pResource
 			&& Subresource == Last.Transition.Subresource
-			&& Before      == Last.Transition.StateAfter
-			&& After       == Last.Transition.StateBefore
-			)
+			// This should always be true (assuming the above are also true) but we don't want to obscure bugs by assuming it
+			&& Before == Last.Transition.StateAfter)
 		{
-			Barriers.RemoveAt(Barriers.Num() - 1);
-			return -1;
+			// Check if we are simply reverting the last transition. In that case, we can just remove both transitions.
+			// This happens fairly frequently due to resource pooling since different RHI buffers can point to the same underlying D3D buffer.
+			// Instead of ping-ponging that underlying resource between COPY_DEST and GENERIC_READ, several copies can happen without a ResourceBarrier() in between.
+			// Doing this check also eliminates a D3D debug layer warning about multiple transitions of the same subresource.
+			if (After == Last.Transition.StateBefore)
+			{
+				Barriers.RemoveAt(Barriers.Num() - 1);
+				return -1;
+			}
+			// No work was done between this and the last transition so we can replace the existing transition's after state with
+			// the new after state, effectively merging the two barriers. This commonly happens during copies since we "push" and 
+			// then "pop" the copy states and often there's no additional work done between that pop (e.g. transition from the copy
+			// state back to whatever state the resource was in before) and the next transition.
+			else
+			{
+				Last.Transition.StateAfter = After;
+				return 0;
+			}
 		}
 	}
 
