Render event tracing

Adds a new Insights channel called "RenderTrace" which records renderer activity: RDG pass execution, command list translation, submission queue processing and so on. The primary use case is crash and hang debugging: the Insights tail buffer can be saved and attached to crash reports, providing additional context for diagnosis. Currently the channel emits events for command list recording and submission; more subchannels will follow.

The data is analyzed and visualized by the RenderTraceInsights plug-in.

#rnx
#Synced-CL 51937798

#Synced-CL 51937798

diff --git a/Engine/Plugins/RenderTraceInsights/Source/RenderTraceInsights/Private/RenderTraceAnalyzer.cpp b/Engine/Plugins/RenderTraceInsights/Source/RenderTraceInsights/Private/RenderTraceAnalyzer.cpp
--- a/Engine/Plugins/RenderTraceInsights/Source/RenderTraceInsights/Private/RenderTraceAnalyzer.cpp
+++ b/Engine/Plugins/RenderTraceInsights/Source/RenderTraceInsights/Private/RenderTraceAnalyzer.cpp
@@ -53,7 +53,7 @@ void FRenderTraceAnalyzer::OnAnalysisBegin(const FOnAnalysisContext& Context)
 	ROUTE_EVENT(EndTranslateCommandList);
 	ROUTE_EVENT(FinalizeTranslateJob);
 	ROUTE_EVENT(SubmitTranslateJobs);
-	ROUTE_EVENT(AssociatePlatformContext);
+	ROUTE_EVENT(D3D12AssociateContext);
 	ROUTE_EVENT(D3D12PayloadCreated);
 	ROUTE_EVENT(D3D12AddPayload);
 	ROUTE_EVENT(D3D12CommandListOp);
@@ -818,7 +818,7 @@ void FRenderTraceAnalyzer::ProcessSetTranslateContext(const FEventData& EventDat
 	// happens. The association will still be valid when it does that, because these extra actions happen before returning the context
 	// to the pool, so another job can't steal the context while we still need the old association.
 	// Example flow with d3d12:
-	//		FRHICommandListExecutor::FTranslateState::Finalize()
+	//		FRHICommandListExecutor::FTranslateState::Translate_CloseChain()
 	//			|-> GDynamicRHI->RHICloseTranslateChain()
 	//			|	|-> FD3D12DynamicRHI::RHIFinalizeContext()
 	//			|		|-> FD3D12ContextCommon::Finalize()
@@ -916,7 +916,7 @@ void FRenderTraceAnalyzer::ProcessFinalizeTranslateJob(const FEventData& EventDa
 void FRenderTraceAnalyzer::ProcessSubmitTranslateJobs(const FEventData& EventData, uint32 ThreadID, double Timestamp)
 {
 	TConstArrayView<uint64> TranslateJobIDs = EventData.GetArrayView<uint64>("TranslateJobIDs");
-	const uint32 RTSubmissionID = EventData.GetValue<uint32>("RTSubmissionID");
+	const uint64 SubmissionID = EventData.GetValue<uint64>("SubmissionID");
 
 	// Submit "tasks" are just events, so we'll use an arbitrary duration, but we need to make sure we don't
 	// overlap the previous submit event on the submission timeline.
@@ -932,19 +932,46 @@ void FRenderTraceAnalyzer::ProcessSubmitTranslateJobs(const FEventData& EventDat
 
 	Timeline.AppendBeginEvent(SubmitTask.StartTime, FRenderTraceEvent(NewTaskEntry.Key));
 
-	// Link the render-thread submit event to the submit task.
-	const uint32* SubmitPassID = PendingRTSubmitIDs.Find(RTSubmissionID);
-	if (SubmitPassID)
+	// If the bottom bit is set, this submit event was instigated by the render thread via ERHISubmitFlags::SubmitToGPU, and the RT
+	// submission ID is in the top 32 bits. Otherwise it was instigated by the completion of a translate chain, and the value contains
+	// the app ID of the last job in the chain.
+	if (SubmissionID & 0x01)
 	{
-		FRDGPassInstance& SubmitPass = Provider.EditRDGPass(*SubmitPassID);
-		check(SubmitPass.Type == ERDGPassType::Submit);
-		SubmitPass.SubmitTaskID = NewTaskEntry.Key;
-		SubmitTask.SubmitPassID = *SubmitPassID;
-		PendingRTSubmitIDs.Remove(RTSubmissionID);
+		// Link the render-thread submit event to the submit task.
+		const uint32 RTSubmissionID = (uint32)(SubmissionID >> 32);
+		const uint32* SubmitPassID = PendingRTSubmitIDs.Find(RTSubmissionID);
+		if (SubmitPassID)
+		{
+			FRDGPassInstance& SubmitPass = Provider.EditRDGPass(*SubmitPassID);
+			check(SubmitPass.Type == ERDGPassType::Submit);
+			SubmitPass.SubmitTaskID = NewTaskEntry.Key;
+
+			SubmitTask.bEagerSubmission = false;
+			SubmitTask.SubmitTaskID = *SubmitPassID;
+
+			PendingRTSubmitIDs.Remove(RTSubmissionID);
+		}
+		else
+		{
+			UE_LOG(LogRenderTrace, Warning, TEXT("Got submit event with unknown render thread pass ID %x, not linking."), RTSubmissionID);
+		}
 	}
 	else
 	{
-		UE_LOG(LogRenderTrace, Warning, TEXT("Got submit event with unknown render thread pass ID %x, not linking."), RTSubmissionID);
+		const FLiveTranslateJob* LiveTranslateJob = LiveTranslateJobs.Find(SubmissionID);
+		if (LiveTranslateJob)
+		{
+			FRHITranslateTask& InstigatorTask = Provider.EditRHITranslateTask(LiveTranslateJob->PermanentID);
+			check(InstigatorTask.Type == ERHITranslateTaskType::Translate);
+			InstigatorTask.SubmitTaskID = NewTaskEntry.Key;
+
+			SubmitTask.bEagerSubmission = true;
+			SubmitTask.SubmitTaskID = LiveTranslateJob->PermanentID;
+		}
+		else
+		{
+			UE_LOG(LogRenderTrace, Warning, TEXT("Got submit event with unknown instigator translate job ID %llx, not linking."), SubmissionID);
+		}
 	}
 
 	// We'll add subevents for all the submitted translate jobs. Since these don't have durations either, we'll just create
@@ -995,7 +1022,7 @@ void FRenderTraceAnalyzer::ProcessSubmitTranslateJobs(const FEventData& EventDat
 	Timeline.AppendEndEvent(SubmitTask.EndTime);
 }
 
-void FRenderTraceAnalyzer::ProcessAssociatePlatformContext(const FEventData& EventData, uint32 ThreadID, double Timestamp)
+void FRenderTraceAnalyzer::ProcessD3D12AssociateContext(const FEventData& EventData, uint32 ThreadID, double Timestamp)
 {
 	// Since platform contexts are pooled, we will only get this message once, when the pooled object is created. These tend to live forever,
 	// so PlatformToRHIContext is not like other mapping tables, which are constantly updated as objects are created and destroyed. Its contents
@@ -1557,7 +1584,7 @@ bool FRenderTraceAnalyzer::OnEvent(uint16 RouteId, EStyle Style, const FOnEventC
 	DISPATCH_EVENT(EndTranslateCommandList);
 	DISPATCH_EVENT(FinalizeTranslateJob);
 	DISPATCH_EVENT(SubmitTranslateJobs);
-	DISPATCH_EVENT(AssociatePlatformContext);
+	DISPATCH_EVENT(D3D12AssociateContext);
 	DISPATCH_EVENT(D3D12PayloadCreated);
 	DISPATCH_EVENT(D3D12AddPayload);
 	DISPATCH_EVENT(D3D12CommandListOp);
diff --git a/Engine/Plugins/RenderTraceInsights/Source/RenderTraceInsights/Private/RenderTraceAnalyzer.h b/Engine/Plugins/RenderTraceInsights/Source/RenderTraceInsights/Private/RenderTraceAnalyzer.h
--- a/Engine/Plugins/RenderTraceInsights/Source/RenderTraceInsights/Private/RenderTraceAnalyzer.h
+++ b/Engine/Plugins/RenderTraceInsights/Source/RenderTraceInsights/Private/RenderTraceAnalyzer.h
@@ -42,7 +42,7 @@ private:
 		Packet_EndTranslateCommandList,
 		Packet_FinalizeTranslateJob,
 		Packet_SubmitTranslateJobs,
-		Packet_AssociatePlatformContext,
+		Packet_D3D12AssociateContext,
 		Packet_D3D12PayloadCreated,
 		Packet_D3D12AddPayload,
 		Packet_D3D12CommandListOp,
@@ -117,7 +117,7 @@ private:
 	DECLARE_EVENT_PROCESSING_FUNC(EndTranslateCommandList);
 	DECLARE_EVENT_PROCESSING_FUNC(FinalizeTranslateJob);
 	DECLARE_EVENT_PROCESSING_FUNC(SubmitTranslateJobs);
-	DECLARE_EVENT_PROCESSING_FUNC(AssociatePlatformContext);
+	DECLARE_EVENT_PROCESSING_FUNC(D3D12AssociateContext);
 	DECLARE_EVENT_PROCESSING_FUNC(D3D12PayloadCreated);
 	DECLARE_EVENT_PROCESSING_FUNC(D3D12AddPayload);
 	DECLARE_EVENT_PROCESSING_FUNC(D3D12CommandListOp);
diff --git a/Engine/Plugins/RenderTraceInsights/Source/RenderTraceInsights/Private/RenderTraceTrack.cpp b/Engine/Plugins/RenderTraceInsights/Source/RenderTraceInsights/Private/RenderTraceTrack.cpp
--- a/Engine/Plugins/RenderTraceInsights/Source/RenderTraceInsights/Private/RenderTraceTrack.cpp
+++ b/Engine/Plugins/RenderTraceInsights/Source/RenderTraceInsights/Private/RenderTraceTrack.cpp
@@ -32,7 +32,8 @@ constexpr FLinearColor GRelationColorSubmitCmdList(1.0f, 0.647f, 0.0f, 1.0f);
 constexpr FLinearColor GRelationColorTranslate(1.0f, 0.0f, 0.0f, 1.0f);
 constexpr FLinearColor GRelationColorSplitFromTranslateJob(0.0f, 0.0f, 1.0f, 1.0f);
 constexpr FLinearColor GRelationColorSplitToTranslateJob(0.0f, 0.2f, 1.0f, 1.0f);
-constexpr FLinearColor GRelationColorSubmitRHI(1.0f, 0.0f, 1.0f, 1.0f);
+constexpr FLinearColor GRelationColorTriggerSubmission(1.0f, 0.58f, 0.12f, 1.0f);
+constexpr FLinearColor GRelationColorSubmitTranslateJob(1.0f, 0.0f, 1.0f, 1.0f);
 constexpr FLinearColor GRelationColorExecute(1.0f, 0.5f, 0.5f, 1.0f);
 constexpr FLinearColor GRelationColorResolveSyncPoint(0.0f, 0.5f, 0.5f, 1.0f);
 constexpr FLinearColor GRelationColorInterrupt(1.0f, 0.5f, 0.0f, 1.0f);
@@ -1032,7 +1033,7 @@ void FRenderThreadSubmissionTrack::OnRenderTraceEventSelected(const FTimingEvent
 		if (RHISubmissionTrack.IsValid())
 		{
 			TUniquePtr<ITimingEventRelation> Relation = MakeUnique<FRenderTraceRelation>(
-				SharedThis(this), InSelectedEvent.GetEndTime(), 0, RHISubmissionTrack, Task.StartTime, 0, GRelationColorSubmitRHI
+				SharedThis(this), InSelectedEvent.GetEndTime(), 0, RHISubmissionTrack, Task.StartTime, 0, GRelationColorTriggerSubmission
 			);
 
 			TimingView->AddRelation(Relation);
@@ -1436,7 +1437,7 @@ void FRHIThreadTrack::OnTranslateEventSelected(const FTimingEvent& InSelectedEve
 		TimingView->AddRelation(Relation);
 	}
 
-	// Draw arrow to submit task.
+	// Draw arrow to our portion of the submit task.
 	if (Task.NextPhaseTaskID != INVALID_EVENT_ID)
 	{
 		const FRHITranslateTask& SubmitTask = Provider->GetRHITranslateTask(Task.NextPhaseTaskID);
@@ -1450,7 +1451,7 @@ void FRHIThreadTrack::OnTranslateEventSelected(const FTimingEvent& InSelectedEve
 				if (TranslateEvent.ItemID == TaskID)
 				{
 					TUniquePtr<ITimingEventRelation> Relation = MakeUnique<FRenderTraceRelation>(
-						SharedThis(this), Task.EndTime, 0, RHISubmissionTrack, TranslateEvent.Start, 1, GRelationColorSubmitRHI
+						SharedThis(this), Task.EndTime, 0, RHISubmissionTrack, TranslateEvent.Start, 1, GRelationColorSubmitTranslateJob
 					);
 
 					TimingView->AddRelation(Relation);
@@ -1459,6 +1460,22 @@ void FRHIThreadTrack::OnTranslateEventSelected(const FTimingEvent& InSelectedEve
 			}
 		}
 	}
+
+	// If we triggered submission, draw an arrow to the submission event.
+	if (Task.SubmitTaskID != INVALID_EVENT_ID)
+	{
+		const FRHITranslateTask& SubmitTask = Provider->GetRHITranslateTask(Task.SubmitTaskID);
+		check(SubmitTask.Type == ERHITranslateTaskType::Submit);
+		TSharedPtr<const FBaseTimingTrack> RHISubmissionTrack = GetFirstChildTrack(SharedData.GetRHIThreadTrack(SubmitTask.ThreadID));
+		if (RHISubmissionTrack.IsValid())
+		{
+			TUniquePtr<ITimingEventRelation> Relation = MakeUnique<FRenderTraceRelation>(
+				SharedThis(this), Task.EndTime, 0, RHISubmissionTrack, SubmitTask.StartTime, 0, GRelationColorTriggerSubmission
+			);
+
+			TimingView->AddRelation(Relation);
+		}
+	}
 	
 	CacheSplitFromTask(*Provider, TaskID);
 	if (Task.SplitFromTranslateTaskID < RELATED_TRANSLATE_JOB_NOT_FOUND)
@@ -1793,17 +1810,32 @@ void FRHISubmissionTrack::OnSubmitEventSelected(const FTimingEvent& InSelectedEv
 	const FRHITranslateTask& SubmitTask = Provider->GetRHITranslateTask(SubmitTaskID);
 	check(SubmitTask.Type == ERHITranslateTaskType::Submit);
 
-	// Draw arrow from instigating render thread task.
-	if (SubmitTask.SubmitPassID != INVALID_EVENT_ID)
+	// Draw arrow from instigating task.
+	if (SubmitTask.SubmitTaskID != INVALID_EVENT_ID)
 	{
-		const FRDGPassInstance& SubmitPass = Provider->GetRDGPass(SubmitTask.SubmitPassID);
-		TSharedPtr<const FBaseTimingTrack> RTSubmissionTrack = GetFirstChildTrack(SharedData.GetRDGThreadTrack(SubmitPass.ExecThreadID));
-		if (RTSubmissionTrack.IsValid())
+		TSharedPtr<const FBaseTimingTrack> InstigatorTrack;
+		double InstigatorTime;
+		if (SubmitTask.bEagerSubmission)
+		{
+			// The instigator is a translate job.
+			const FRHITranslateTask& InstigatorTask = Provider->GetRHITranslateTask(SubmitTask.SubmitTaskID);
+			check(InstigatorTask.Type == ERHITranslateTaskType::Translate);
+			InstigatorTrack = SharedData.GetRHIThreadTrack(InstigatorTask.ThreadID);
+			InstigatorTime = InstigatorTask.EndTime;
+		}
+		else
+		{
+			// The instigator is a render thread submit event.
+			const FRDGPassInstance& SubmitPass = Provider->GetRDGPass(SubmitTask.SubmitTaskID);
+			InstigatorTrack = GetFirstChildTrack(SharedData.GetRDGThreadTrack(SubmitPass.ExecThreadID));
+			InstigatorTime = SubmitPass.EndTime;
+		}
+
+		if (InstigatorTrack.IsValid())
 		{
 			TUniquePtr<ITimingEventRelation> Relation = MakeUnique<FRenderTraceRelation>(
-				RTSubmissionTrack, SubmitPass.EndTime, 0, SharedThis(this), SubmitTask.StartTime, 0, GRelationColorSubmitRHI
+				InstigatorTrack, InstigatorTime, 0, SharedThis(this), SubmitTask.StartTime, 0, GRelationColorTriggerSubmission
 			);
-
 			TimingView->AddRelation(Relation);
 		}
 	}
@@ -1817,7 +1849,7 @@ void FRHISubmissionTrack::OnSubmitEventSelected(const FTimingEvent& InSelectedEv
 		if (TranslateTrack.IsValid())
 		{
 			TUniquePtr<ITimingEventRelation> Relation = MakeUnique<FRenderTraceRelation>(
-				TranslateTrack, TranslateTask.EndTime, 0, SharedThis(this), TranslateEvent.Start, 1, GRelationColorSubmitRHI
+				TranslateTrack, TranslateTask.EndTime, 0, SharedThis(this), TranslateEvent.Start, 1, GRelationColorSubmitTranslateJob
 			);
 
 			TimingView->AddRelation(Relation);
@@ -1843,7 +1875,7 @@ void FRHISubmissionTrack::OnSubmitTranslateJobEventSelected(const FTimingEvent&
 	}
 
 	TUniquePtr<ITimingEventRelation> Relation = MakeUnique<FRenderTraceRelation>(
-		TranslateTrack, TranslateTask.EndTime, 0, SharedThis(this), InSelectedEvent.GetStartTime(), 1, GRelationColorSubmitRHI
+		TranslateTrack, TranslateTask.EndTime, 0, SharedThis(this), InSelectedEvent.GetStartTime(), 1, GRelationColorSubmitTranslateJob
 	);
 
 	TimingView->AddRelation(Relation);
@@ -1915,13 +1947,22 @@ void FRHISubmissionTrack::PopulateSubmitTaskDetails(IEventDetailsBuilder& Builde
 	Builder.AddTitle(TEXTVIEW("Submit Task"));
 	Builder.AddNameValueTextLine(TEXTVIEW("ID:"), LexToString(Task.AppID));
 
-	if (Task.SubmitPassID != INVALID_EVENT_ID)
+	if (Task.SubmitTaskID != INVALID_EVENT_ID)
 	{
-		Builder.AddNameValueTextLine(TEXTVIEW("Triggered By:"), LexToString(Task.SubmitPassID));
+		if (Task.bEagerSubmission)
+		{
+			const FRHITranslateTask& InstigatorTask = Provider.GetRHITranslateTask(Task.SubmitTaskID);
+			check(InstigatorTask.Type == ERHITranslateTaskType::Translate);
+			Builder.AddNameValueTextLine(TEXTVIEW("Triggered By:"), FString::Printf(TEXT("Translate job 0x%llx"), InstigatorTask.AppID));
+		}
+		else
+		{
+			Builder.AddNameValueTextLine(TEXTVIEW("Triggered By:"), FString::Printf(TEXT("RT submit event %u"), Task.SubmitTaskID));
+		}
 	}
 	else
 	{
-		Builder.AddNameValueTextLine(TEXTVIEW("NOT TRIGGERED"), TEXTVIEW(""));
+		Builder.AddNameValueTextLine(TEXTVIEW("NO INSTIGATOR INFO"), TEXTVIEW(""));
 	}
 
 	FAppIDListBuilder<FRHITranslateTask>::AddAllEvents(Builder, TEXT("Translate Jobs:"), Task.ProcessedItems);
diff --git a/Engine/Plugins/RenderTraceInsights/Source/RenderTraceInsights/Public/IRenderTraceProvider.h b/Engine/Plugins/RenderTraceInsights/Source/RenderTraceInsights/Public/IRenderTraceProvider.h
--- a/Engine/Plugins/RenderTraceInsights/Source/RenderTraceInsights/Public/IRenderTraceProvider.h
+++ b/Engine/Plugins/RenderTraceInsights/Source/RenderTraceInsights/Public/IRenderTraceProvider.h
@@ -255,8 +255,11 @@ struct FRHITranslateTask
 
 	uint8 JobFlags = 0; // ERHITranslateJobFlags
 
-	// The render-thread submit pass ID for submit tasks.
-	uint32 SubmitPassID = INVALID_EVENT_ID;
+	// For a submit task, this identifies the render-thread submit pass ID if the instigator is the render thread,
+	// or the translate job if it's an eager submit. For a translate tasks, this identifies the submission it
+	// triggered, if any.
+	uint32 SubmitTaskID = INVALID_EVENT_ID;
+	bool bEagerSubmission = false;
 	
 	// Previous and next jobs when a translate chain is split. This is not computed during analysis because it's too complicated.
 	// Instead, we compute and cache these values on demand when an event is selected.
diff --git a/Engine/Source/Runtime/D3D12RHI/Private/D3D12CommandContext.cpp b/Engine/Source/Runtime/D3D12RHI/Private/D3D12CommandContext.cpp
--- a/Engine/Source/Runtime/D3D12RHI/Private/D3D12CommandContext.cpp
+++ b/Engine/Source/Runtime/D3D12RHI/Private/D3D12CommandContext.cpp
@@ -10,6 +10,8 @@ D3D12CommandContext.cpp: RHI  Command Context implementation.
 #include "D3D12AmdExtensions.h"
 #include "D3D12RayTracing.h"
 
+#include "D3D12Tracing.h"
+
 int32 GD3D12MaxCommandsPerCommandList = 10000;
 static FAutoConsoleVariableRef CVarMaxCommandsPerCommandList(
 	TEXT("D3D12.MaxCommandsPerCommandList"),
@@ -68,6 +70,12 @@ FD3D12CommandContext::FD3D12CommandContext(FD3D12Device* InParent, ED3D12QueueTy
 {
 	StaticUniformBuffers.AddZeroed(FUniformBufferStaticSlotRegistry::Get().GetSlotCount());
 	ClearState();
+
+	// All the interesting stuff (sync points, payloads etc.) happens in FD3D12ContextCommon, but the tracing code associates
+	// translate jobs with IRHIComputeContext pointers, and FD3D12ContextCommon does not derive from that interface. Instead,
+	// this class derives from both, so this is the place where we can map FD3D12ContextCommon pointers to IRHIComputeContext
+	// pointers, to allow the tracing code to reconstruct the translate job -> RHI context -> D3D12 context dependency chains.
+	IF_RENDER_TRACING_ENABLED(RenderTracing::D3D12AssociateContext(this, this));
 }
 
 FD3D12CommandContext::~FD3D12CommandContext()
@@ -235,7 +243,10 @@ void FD3D12ContextCommon::WaitSyncPoint(FD3D12SyncPoint* SyncPoint)
 		CloseCommandList();
 	}
 
-	GetPayload(EPhase::Wait)->SyncPointsToWait.Add(SyncPoint);
+	FD3D12Payload* Payload = GetPayload(EPhase::Wait);
+	Payload->SyncPointsToWait.Add(SyncPoint);
+
+	IF_RENDER_TRACING_ENABLED(RenderTracing::D3D12WaitSyncPoint(Payload, SyncPoint));
 }
 
 void FD3D12ContextCommon::SignalSyncPoint(FD3D12SyncPoint* SyncPoint)
@@ -245,7 +256,10 @@ void FD3D12ContextCommon::SignalSyncPoint(FD3D12SyncPoint* SyncPoint)
 		CloseCommandList();
 	}
 
-	GetPayload(EPhase::Signal)->SyncPointsToSignal.Add(SyncPoint);
+	FD3D12Payload* Payload = GetPayload(EPhase::Signal);
+	Payload->SyncPointsToSignal.Add(SyncPoint);
+
+	IF_RENDER_TRACING_ENABLED(RenderTracing::D3D12SignalSyncPoint(Payload, SyncPoint));
 }
 
 void FD3D12ContextCommon::SignalManualFence(ID3D12Fence* Fence, uint64 Value)
@@ -255,7 +269,10 @@ void FD3D12ContextCommon::SignalManualFence(ID3D12Fence* Fence, uint64 Value)
 		CloseCommandList();
 	}
 
-	GetPayload(EPhase::Signal)->ManualFencesToSignal.Emplace(Fence, Value, nullptr);
+	FD3D12Payload* Payload = GetPayload(EPhase::Signal);
+	Payload->ManualFencesToSignal.Emplace(Fence, Value, nullptr);
+
+	IF_RENDER_TRACING_ENABLED(RenderTracing::D3D12SignalManualFence(Payload, Fence, Value));
 }
 
 void FD3D12ContextCommon::WaitManualFence(ID3D12Fence* Fence, uint64 Value, const TSharedPtr<std::atomic<bool>>& bHasSignalled)
@@ -265,7 +282,10 @@ void FD3D12ContextCommon::WaitManualFence(ID3D12Fence* Fence, uint64 Value, cons
 		CloseCommandList();
 	}
 
-	GetPayload(EPhase::Wait)->ManualFencesToWait.Emplace(Fence, Value, bHasSignalled);
+	FD3D12Payload* Payload = GetPayload(EPhase::Wait);
+	Payload->ManualFencesToWait.Emplace(Fence, Value, bHasSignalled);
+
+	IF_RENDER_TRACING_ENABLED(RenderTracing::D3D12WaitManualFence(Payload, Fence, Value));
 }
 
 FD3D12QueryLocation FD3D12ContextCommon::AllocateQuery(ED3D12QueryType Type, void* Target)
@@ -338,7 +358,10 @@ void FD3D12ContextCommon::OpenCommandList()
 
 	// Get a new command list
 	CommandList = Device->ObtainCommandList(CommandAllocator, &TimestampQueries, &PipelineStatsQueries);
-	GetPayload(EPhase::Execute)->CommandListsToExecute.Add(CommandList);
+	FD3D12Payload* Payload = GetPayload(EPhase::Execute);
+	Payload->CommandListsToExecute.Add(CommandList);
+
+	IF_RENDER_TRACING_ENABLED(RenderTracing::D3D12OpenCommandList(Payload, CommandList->BaseCommandList().GetNoRefCount()));
 
 	check(ActiveQueries == 0);
 }
@@ -366,6 +389,8 @@ void FD3D12ContextCommon::CloseCommandList()
 
 	FD3D12Payload* Payload = GetPayload(EPhase::Execute);
 
+	IF_RENDER_TRACING_ENABLED(RenderTracing::D3D12CloseCommandList(Payload, CommandList->BaseCommandList().GetNoRefCount()));
+
 	CommandList->Close();
 	CommandList = nullptr;
 
@@ -409,13 +434,29 @@ void FD3D12ContextCommon::Finalize(TArray<FD3D12Payload*>& OutPayloads)
 			? Payloads[0]
 			: GetPayload(EPhase::Wait);
 
+#if UE_RENDER_TRACING_ENABLED
+		for (const FD3D12SyncPointRef& SyncPoint : BatchedSyncPoints.ToWait)
+		{
+			RenderTracing::D3D12WaitSyncPoint(Payload, SyncPoint);
+		}
+#endif
+
 		Payload->SyncPointsToWait.Append(BatchedSyncPoints.ToWait);
 		BatchedSyncPoints.ToWait.Reset();
 	}
 
 	if (BatchedSyncPoints.ToSignal.Num())
 	{
-		GetPayload(EPhase::Signal)->SyncPointsToSignal.Append(BatchedSyncPoints.ToSignal);
+		FD3D12Payload* Payload = GetPayload(EPhase::Signal);
+
+#if UE_RENDER_TRACING_ENABLED
+		for (const FD3D12SyncPointRef& SyncPoint : BatchedSyncPoints.ToSignal)
+		{
+			RenderTracing::D3D12SignalSyncPoint(Payload, SyncPoint);
+		}
+#endif
+
+		Payload->SyncPointsToSignal.Append(BatchedSyncPoints.ToSignal);
 		BatchedSyncPoints.ToSignal.Reset();
 	}
 
@@ -574,7 +615,9 @@ FD3D12SyncPoint* FD3D12CopyScope::GetSyncPoint() const
 
 void FD3D12ContextCommon::NewPayload()
 {
-	Payloads.Add(new FD3D12Payload(Device->GetQueue(QueueType)));
+	FD3D12Payload* Payload = new FD3D12Payload(Device->GetQueue(QueueType));
+	IF_RENDER_TRACING_ENABLED(RenderTracing::D3D12AddPayload(this, Payload));
+	Payloads.Add(Payload);
 }
 
 // Set the Present Params in the context for the Submission Thread to call the native Present
diff --git a/Engine/Source/Runtime/D3D12RHI/Private/D3D12Device.cpp b/Engine/Source/Runtime/D3D12RHI/Private/D3D12Device.cpp
--- a/Engine/Source/Runtime/D3D12RHI/Private/D3D12Device.cpp
+++ b/Engine/Source/Runtime/D3D12RHI/Private/D3D12Device.cpp
@@ -5,6 +5,7 @@
 #include "D3D12IntelExtensions.h"
 #include "D3D12RayTracing.h"
 #include "D3D12ExplicitDescriptorCache.h"
+#include "D3D12Tracing.h"
 
 static TAutoConsoleVariable<int32> CVarD3D12GPUTimeout(
 	TEXT("r.D3D12.GPUTimeout"),
@@ -561,6 +562,8 @@ void FD3D12Device::BlockUntilIdle()
 
 		Payloads.Add(Payload);
 		SyncPoints.Add(SyncPoint);
+
+		IF_RENDER_TRACING_ENABLED(RenderTracing::D3D12SignalSyncPoint(Payload, SyncPoint));
 	}
 
 	FD3D12DynamicRHI::GetD3DRHI()->SubmitPayloads(MoveTemp(Payloads));
diff --git a/Engine/Source/Runtime/D3D12RHI/Private/D3D12RHI.cpp b/Engine/Source/Runtime/D3D12RHI/Private/D3D12RHI.cpp
--- a/Engine/Source/Runtime/D3D12RHI/Private/D3D12RHI.cpp
+++ b/Engine/Source/Runtime/D3D12RHI/Private/D3D12RHI.cpp
@@ -14,6 +14,8 @@
 #include "D3D12AmdExtensions.h"
 #include "D3D12IntelExtensions.h"
 
+#include "D3D12Tracing.h"
+
 #if !defined(D3D12_PLATFORM_NEEDS_DISPLAY_MODE_ENUMERATION)
 	#define D3D12_PLATFORM_NEEDS_DISPLAY_MODE_ENUMERATION 1
 #endif
@@ -427,6 +429,8 @@ void FD3D12DynamicRHI::EnqueueEndOfPipeTask(TUniqueFunction<void()> TaskFunc, TU
 		Payload->SyncPointsToSignal.Emplace(SyncPoint);
 		Prereqs.Add(SyncPoint->GetGraphEvent());
 
+		IF_RENDER_TRACING_ENABLED(RenderTracing::D3D12SignalSyncPoint(Payload, SyncPoint));
+
 		if (ModifyPayloadCallback)
 			ModifyPayloadCallback(*Payload);
 
diff --git a/Engine/Source/Runtime/D3D12RHI/Private/D3D12Submission.cpp b/Engine/Source/Runtime/D3D12RHI/Private/D3D12Submission.cpp
--- a/Engine/Source/Runtime/D3D12RHI/Private/D3D12Submission.cpp
+++ b/Engine/Source/Runtime/D3D12RHI/Private/D3D12Submission.cpp
@@ -7,6 +7,7 @@
 #include "IRenderCaptureProvider.h"
 #include "Stats/ThreadIdleStats.h"
 #include "ProfilingDebugging/AssetMetadataTrace.h"
+#include "D3D12Tracing.h"
 
 #ifndef D3D12_PLATFORM_SUPPORTS_BLOCKING_FENCES
 #define D3D12_PLATFORM_SUPPORTS_BLOCKING_FENCES 1
@@ -333,6 +334,8 @@ FD3D12DynamicRHI::FProcessResult FD3D12DynamicRHI::ProcessSubmissionQueue()
 	SCOPED_NAMED_EVENT_TEXT("SubmissionQueue_Process", FColor::Turquoise);
 	SCOPE_CYCLE_COUNTER(STAT_D3D12Submit);
 	LLM_SCOPE_BYNAME(TEXT("RHIMisc/ProcessSubmissionQueue"));
+	
+	IF_RENDER_TRACING_ENABLED(RenderTracing::D3D12ProcessSubmissionQueueEnter());
 
 	FD3D12Queue::FPayloadArray PayloadsToHandDown;
 
@@ -403,6 +406,7 @@ FD3D12DynamicRHI::FProcessResult FD3D12DynamicRHI::ProcessSubmissionQueue()
 							// Skip processing this queue, and move on to the next. We will retry later when
 							// further work is submitted, which may contain the sync point we need.
 							Result.Status |= EQueueStatus::Pending;
+							IF_RENDER_TRACING_ENABLED(RenderTracing::D3D12SubmitYieldOnUnresolvedSyncPoint(CurrentQueue, Payload, SyncPoint));
 							return;
 						}
 
@@ -425,6 +429,7 @@ FD3D12DynamicRHI::FProcessResult FD3D12DynamicRHI::ProcessSubmissionQueue()
 							// Skip processing this queue, and move on to the next. We will retry later when
 							// further work is submitted, which may contain the fence we need.
 							Result.Status |= EQueueStatus::Pending;
+							IF_RENDER_TRACING_ENABLED(RenderTracing::D3D12SubmitYieldOnUnresolvedManualFence(CurrentQueue, Payload, ManualFence.Fence, ManualFence.Value));
 							return;
 						}
 
@@ -481,6 +486,7 @@ FD3D12DynamicRHI::FProcessResult FD3D12DynamicRHI::ProcessSubmissionQueue()
 		InterruptThread->Kick();
 	}
 
+	IF_RENDER_TRACING_ENABLED(RenderTracing::D3D12ProcessSubmissionQueueExit(static_cast<uint8>(Result.Status)));
 	return Result;
 }
 
@@ -613,6 +619,7 @@ uint64 FD3D12Queue::FinalizePayload(bool bRequiresSignal, FPayloadArray& Payload
 	for (FD3D12SyncPointRef& SyncPoint : PayloadToSubmit->SyncPointsToSignal)
 	{
 		check(!SyncPoint->ResolvedFence.IsSet());
+		IF_RENDER_TRACING_ENABLED(RenderTracing::D3D12SubmitResolveSyncPoint(*this, PayloadToSubmit, SyncPoint, Fence.NextCompletionValue));
 		SyncPoint->ResolvedFence.Emplace(Fence, NextCompletionValue);
 	}
 
@@ -684,6 +691,7 @@ void FD3D12DynamicRHI::FlushBatchedPayloads(FD3D12Queue::FPayloadArray& Payloads
 				, LocalFence.OwnerQueue->GetProfilerQueue()
 			);
 
+			IF_RENDER_TRACING_ENABLED(RenderTracing::D3D12SubmitWaitQueueFence(Payload, Queue, *LocalFence.OwnerQueue, Value));
 			VERIFYD3D12RESULT(Queue.D3DCommandQueue->Wait(LocalFence.D3DFence, Value));
 		}
 
@@ -691,6 +699,7 @@ void FD3D12DynamicRHI::FlushBatchedPayloads(FD3D12Queue::FPayloadArray& Payloads
 		for (auto& [LocalFence, Value, bFenceSignalled] : Payload->ManualFencesToWait)
 		{
 			checkf(!bFenceSignalled.IsValid() || *bFenceSignalled, TEXT("The signal operation for a provided manual D3D12 fence has not yet happened."));
+			IF_RENDER_TRACING_ENABLED(RenderTracing::D3D12SubmitWaitManualFence(Payload, Queue, LocalFence, Value));
 			VERIFYD3D12RESULT(Queue.D3DCommandQueue->Wait(LocalFence, Value));
 		}
 	};
@@ -731,6 +740,12 @@ void FD3D12DynamicRHI::FlushBatchedPayloads(FD3D12Queue::FPayloadArray& Payloads
 			CommandLists.Append(MoveTemp(Payload->CommandListsToExecute));
 		}
 
+		IF_RENDER_TRACING_ENABLED(RenderTracing::D3D12SubmitPlatformCommandLists(
+			Queue,
+			MakeConstArrayView<FD3D12Payload*>(&PayloadsToSubmit[FirstPayload], LastPayload - FirstPayload),
+			D3DCommandLists
+		));
+
 		const int32 MaxBatchSize = GetMaxExecuteBatchSize(Queue.QueueType);
 		const int32 NumCommandLists = D3DCommandLists.Num();
 
@@ -783,6 +798,8 @@ void FD3D12DynamicRHI::FlushBatchedPayloads(FD3D12Queue::FPayloadArray& Payloads
 		// Signal any manual fences
 		for (auto& [ManualFence, Value, bFenceSignalled] : Payload->ManualFencesToSignal)
 		{
+			IF_RENDER_TRACING_ENABLED(RenderTracing::D3D12SubmitSignalManualFence(Payload, Queue, ManualFence, Value));
+
 			VERIFYD3D12RESULT(Queue.D3DCommandQueue->Signal(ManualFence, Value));
 			if (bFenceSignalled)
 			{
@@ -800,6 +817,8 @@ void FD3D12DynamicRHI::FlushBatchedPayloads(FD3D12Queue::FPayloadArray& Payloads
 				, Payload->CompletionFenceValue
 			);
 
+			IF_RENDER_TRACING_ENABLED(RenderTracing::D3D12SubmitSignalQueueFence(Payload, Queue, Payload->CompletionFenceValue));
+
 			VERIFYD3D12RESULT(Queue.D3DCommandQueue->Signal(Queue.Fence.D3DFence, Payload->CompletionFenceValue));
 			Queue.Fence.LastSignaledValue.store(Payload->CompletionFenceValue, std::memory_order_release);
 		}
@@ -1038,6 +1057,8 @@ FD3D12DynamicRHI::FProcessResult FD3D12DynamicRHI::ProcessInterruptQueue()
 	SCOPED_NAMED_EVENT_TEXT("InterruptQueue_Process", FColor::Yellow);
 	LLM_SCOPE_BYNAME(TEXT("RHIMisc/ProcessInterruptQueue"));
 
+	IF_RENDER_TRACING_ENABLED(RenderTracing::D3D12ProcessInterruptQueueEnter());
+
 	// Timer that clamps each tick to prevent false positive GPU timeouts
 	// when a debugger is attached and the process is broken.
 	struct FTimer
@@ -1149,6 +1170,8 @@ FD3D12DynamicRHI::FProcessResult FD3D12DynamicRHI::ProcessInterruptQueue()
 			CurrentQueue.PendingInterrupt.Dequeue();
 			Result.Status |= EQueueStatus::Processed;
 
+			IF_RENDER_TRACING_ENABLED(RenderTracing::D3D12InterruptQueueFenceSignaled(CurrentQueue, Payload, CompletedFenceValue, LastSignaledFenceValue));
+
 			// Resolve query results
 			{
 				for (FD3D12QueryLocation& Query : Payload->BatchedObjects.OcclusionQueries)
@@ -1249,13 +1272,15 @@ FD3D12DynamicRHI::FProcessResult FD3D12DynamicRHI::ProcessInterruptQueue()
 		TerminateOnGPUCrash(GGPUD3DResult);
 	}
 
+	IF_RENDER_TRACING_ENABLED(RenderTracing::D3D12ProcessInterruptQueueExit(static_cast<uint8>(Result.Status)));
 	return Result;
 }
 
 FD3D12PayloadBase::FD3D12PayloadBase(FD3D12Queue& Queue)
 	: Queue(Queue)
 	, EventStream(Queue.GetProfilerQueue())
-{}
+{
+}
 
 FD3D12PayloadBase::~FD3D12PayloadBase()
 {
diff --git a/Engine/Source/Runtime/D3D12RHI/Private/D3D12Tracing.cpp b/Engine/Source/Runtime/D3D12RHI/Private/D3D12Tracing.cpp
new file mode 100644
--- /dev/null
+++ b/Engine/Source/Runtime/D3D12RHI/Private/D3D12Tracing.cpp
@@ -0,0 +1,420 @@
+// Copyright Epic Games, Inc. All Rights Reserved.
+
+#include "D3D12Tracing.h"
+#include "Trace/Trace.inl"
+#include "HAL/PlatformTime.h"
+#include "D3D12Device.h"
+#include "D3D12Submission.h"
+
+#if UE_RENDER_TRACING_ENABLED
+
+RENDER_TRACE_EVENT_BEGIN(D3D12AssociateContext, Important | NoSync)
+	UE_TRACE_EVENT_FIELD(uint64, PlatformContextID)
+	UE_TRACE_EVENT_FIELD(uint64, RHIContextID)
+UE_TRACE_EVENT_END()
+
+RENDER_TRACE_EVENT_BEGIN(D3D12PayloadCreated)
+	UE_TRACE_EVENT_FIELD(uint64, PayloadID)
+	UE_TRACE_EVENT_FIELD(uint8, Queue)
+UE_TRACE_EVENT_END()
+
+RENDER_TRACE_EVENT_BEGIN(D3D12AddPayload)
+	UE_TRACE_EVENT_FIELD(uint64, PlatformContextID)
+	UE_TRACE_EVENT_FIELD(uint64, PayloadID)
+UE_TRACE_EVENT_END()
+
+enum ED3D12CommandListOpType
+{
+	D3D12CommandListOp_Open = 0,
+	D3D12CommandListOp_Close = 1
+};
+
+RENDER_TRACE_EVENT_BEGIN(D3D12CommandListOp)
+	UE_TRACE_EVENT_FIELD(uint64, PayloadID)
+	UE_TRACE_EVENT_FIELD(uint64, D3D12CmdListID)
+	UE_TRACE_EVENT_FIELD(uint8, Op)
+UE_TRACE_EVENT_END()
+
+enum ED3D12FenceOpType
+{
+	D3D12FenceOp_Wait = 0,
+	D3D12FenceOp_Signal = 1
+};
+
+RENDER_TRACE_EVENT_BEGIN(D3D12SyncPointOp)
+	UE_TRACE_EVENT_FIELD(uint64, PayloadID)
+	UE_TRACE_EVENT_FIELD(uint64, SyncPointID)
+	UE_TRACE_EVENT_FIELD(uint8, SyncPointType)
+	UE_TRACE_EVENT_FIELD(uint8, Op)
+UE_TRACE_EVENT_END()
+
+RENDER_TRACE_EVENT_BEGIN(D3D12ManualFenceOp)
+	UE_TRACE_EVENT_FIELD(uint64, PayloadID)
+	UE_TRACE_EVENT_FIELD(uint64, FenceID)
+	UE_TRACE_EVENT_FIELD(uint64, FenceVal)
+	UE_TRACE_EVENT_FIELD(uint8, Op)
+UE_TRACE_EVENT_END()
+
+RENDER_TRACE_EVENT_BEGIN(D3D12ProcessSubmissionQueueEnter)
+UE_TRACE_EVENT_END()
+
+RENDER_TRACE_EVENT_BEGIN(D3D12SubmitYieldOnUnresolvedSyncPoint)
+	UE_TRACE_EVENT_FIELD(uint8, Queue)
+	UE_TRACE_EVENT_FIELD(uint64, PayloadID)
+	UE_TRACE_EVENT_FIELD(uint64, SyncPointID)
+UE_TRACE_EVENT_END()
+
+RENDER_TRACE_EVENT_BEGIN(D3D12SubmitYieldOnUnresolvedManualFence)
+	UE_TRACE_EVENT_FIELD(uint8, Queue)
+	UE_TRACE_EVENT_FIELD(uint64, PayloadID)
+	UE_TRACE_EVENT_FIELD(uint64, FenceID)
+	UE_TRACE_EVENT_FIELD(uint64, FenceValue)
+UE_TRACE_EVENT_END()
+
+RENDER_TRACE_EVENT_BEGIN(D3D12SubmitWaitQueueFence)
+	UE_TRACE_EVENT_FIELD(uint64, PayloadID)
+	UE_TRACE_EVENT_FIELD(uint8, ExecutingQueue)
+	UE_TRACE_EVENT_FIELD(uint8, WaitOnQueue)
+	UE_TRACE_EVENT_FIELD(uint64, Value)
+UE_TRACE_EVENT_END()
+
+RENDER_TRACE_EVENT_BEGIN(D3D12SubmitWaitManualFence)
+	UE_TRACE_EVENT_FIELD(uint64, PayloadID)
+	UE_TRACE_EVENT_FIELD(uint8, ExecutingQueue)
+	UE_TRACE_EVENT_FIELD(uint64, FenceID)
+	UE_TRACE_EVENT_FIELD(uint64, Value)
+UE_TRACE_EVENT_END()
+
+RENDER_TRACE_EVENT_BEGIN(D3D12SubmitPlatformCommandLists)
+	UE_TRACE_EVENT_FIELD(uint8, Queue)
+	UE_TRACE_EVENT_FIELD(uint64[], PayloadIDs)
+	UE_TRACE_EVENT_FIELD(uint64[], CmdListIDs)
+UE_TRACE_EVENT_END()
+
+RENDER_TRACE_EVENT_BEGIN(D3D12SubmitSignalManualFence)
+	UE_TRACE_EVENT_FIELD(uint64, PayloadID)
+	UE_TRACE_EVENT_FIELD(uint8, Queue)
+	UE_TRACE_EVENT_FIELD(uint64, FenceID)
+	UE_TRACE_EVENT_FIELD(uint64, Value)
+UE_TRACE_EVENT_END()
+
+RENDER_TRACE_EVENT_BEGIN(D3D12SubmitSignalQueueFence)
+	UE_TRACE_EVENT_FIELD(uint64, PayloadID)
+	UE_TRACE_EVENT_FIELD(uint8, Queue)
+	UE_TRACE_EVENT_FIELD(uint64, Value)
+UE_TRACE_EVENT_END()
+
+RENDER_TRACE_EVENT_BEGIN(D3D12SubmitResolveSyncPoint)
+	UE_TRACE_EVENT_FIELD(uint8, Queue)
+	UE_TRACE_EVENT_FIELD(uint64, PayloadID)
+	UE_TRACE_EVENT_FIELD(uint64, SyncPointID)
+	UE_TRACE_EVENT_FIELD(uint64, Value)
+UE_TRACE_EVENT_END()
+
+RENDER_TRACE_EVENT_BEGIN(D3D12ProcessSubmissionQueueExit)
+	UE_TRACE_EVENT_FIELD(uint8, Status)
+UE_TRACE_EVENT_END()
+
+RENDER_TRACE_EVENT_BEGIN(D3D12ProcessInterruptQueueEnter)
+UE_TRACE_EVENT_END()
+
+RENDER_TRACE_EVENT_BEGIN(D3D12InterruptQueueFenceSignaled)
+	UE_TRACE_EVENT_FIELD(uint8, Queue)
+	UE_TRACE_EVENT_FIELD(uint64, PendingPayloadID)
+	UE_TRACE_EVENT_FIELD(uint64, CurrentFenceValue)
+	UE_TRACE_EVENT_FIELD(uint64, LastCPUSignaledFenceValue)
+UE_TRACE_EVENT_END()
+
+RENDER_TRACE_EVENT_BEGIN(D3D12ProcessInterruptQueueExit)
+	UE_TRACE_EVENT_FIELD(uint8, Status)
+UE_TRACE_EVENT_END()
+
+namespace RenderTracing
+{
+
+static uint8 GetRenderTraceQueueIndex(const FD3D12Queue& Queue)
+{
+	// No mGPU support for now, and only one queue per type.
+	switch (Queue.QueueType)
+	{
+	case ED3D12QueueType::Direct: return 0;
+	case ED3D12QueueType::Async: return 1;
+	case ED3D12QueueType::Copy: return 2;
+	default: checkNoEntry(); return 0;
+	}
+}
+
+void D3D12AssociateContext(const FD3D12ContextCommon* D3D12Context, const IRHIComputeContext* RHIContext)
+{
+	if (!IsEnabled(ERenderTracingChannels::RHISubmission))
+	{
+		return;
+	}
+
+	RENDER_TRACE_LOG(D3D12AssociateContext)
+		<< RENDER_TRACE_VALUE(D3D12AssociateContext, PlatformContextID, reinterpret_cast<uintptr_t>(D3D12Context))
+		<< RENDER_TRACE_VALUE(D3D12AssociateContext, RHIContextID, reinterpret_cast<uintptr_t>(RHIContext));
+}
+
+void D3D12PayloadCreated(const FD3D12Payload* Payload, const FD3D12Queue& Queue)
+{
+	if (!IsEnabled(ERenderTracingChannels::RHISubmission))
+	{
+		return;
+	}
+
+	RENDER_TRACE_LOG(D3D12PayloadCreated)
+		<< RENDER_TRACE_VALUE(D3D12PayloadCreated, PayloadID, reinterpret_cast<uintptr_t>(Payload))
+		<< RENDER_TRACE_VALUE(D3D12PayloadCreated, Queue, GetRenderTraceQueueIndex(Queue));
+}
+
+void D3D12AddPayload(const FD3D12ContextCommon* D3D12Context, const FD3D12Payload* Payload)
+{
+	if (!IsEnabled(ERenderTracingChannels::RHISubmission))
+	{
+		return;
+	}
+
+	RENDER_TRACE_LOG(D3D12AddPayload)
+		<< RENDER_TRACE_VALUE(D3D12AddPayload, PlatformContextID, reinterpret_cast<uintptr_t>(D3D12Context))
+		<< RENDER_TRACE_VALUE(D3D12AddPayload, PayloadID, reinterpret_cast<uintptr_t>(Payload));
+}
+
+static void D3D12CommandListOp(const FD3D12Payload* Payload, const ID3D12CommandList* D3D12CmdList, ED3D12CommandListOpType OpType)
+{
+	if (!IsEnabled(ERenderTracingChannels::RHISubmission))
+	{
+		return;
+	}
+
+	RENDER_TRACE_LOG(D3D12CommandListOp)
+		<< RENDER_TRACE_VALUE(D3D12CommandListOp, PayloadID, reinterpret_cast<uintptr_t>(Payload))
+		<< RENDER_TRACE_VALUE(D3D12CommandListOp, D3D12CmdListID, reinterpret_cast<uintptr_t>(D3D12CmdList))
+		<< RENDER_TRACE_VALUE(D3D12CommandListOp, Op, OpType);
+}
+
+void D3D12OpenCommandList(const FD3D12Payload* Payload, const ID3D12CommandList* D3D12CmdList)
+{
+	D3D12CommandListOp(Payload, D3D12CmdList, D3D12CommandListOp_Open);
+}
+
+void D3D12CloseCommandList(const FD3D12Payload* Payload, const ID3D12CommandList* D3D12CmdList)
+{
+	D3D12CommandListOp(Payload, D3D12CmdList, D3D12CommandListOp_Close);
+}
+
+static void D3D12SyncPointOp(const FD3D12Payload* Payload, const FD3D12SyncPoint* SyncPoint, ED3D12FenceOpType OpType)
+{
+	if (!IsEnabled(ERenderTracingChannels::RHISubmission))
+	{
+		return;
+	}
+
+	RENDER_TRACE_LOG(D3D12SyncPointOp)
+		<< RENDER_TRACE_VALUE(D3D12SyncPointOp, PayloadID, reinterpret_cast<uintptr_t>(Payload))
+		<< RENDER_TRACE_VALUE(D3D12SyncPointOp, SyncPointID, reinterpret_cast<uintptr_t>(SyncPoint))
+		<< RENDER_TRACE_VALUE(D3D12SyncPointOp, SyncPointType, static_cast<uint8>(SyncPoint->GetType()))
+		<< RENDER_TRACE_VALUE(D3D12SyncPointOp, Op, OpType);
+}
+
+void D3D12WaitSyncPoint(const FD3D12Payload* Payload, const FD3D12SyncPoint* SyncPoint)
+{
+	D3D12SyncPointOp(Payload, SyncPoint, D3D12FenceOp_Wait);
+}
+
+void D3D12SignalSyncPoint(const FD3D12Payload* Payload, const FD3D12SyncPoint* SyncPoint)
+{
+	D3D12SyncPointOp(Payload, SyncPoint, D3D12FenceOp_Signal);
+}
+
+static void D3D12ManualFenceOp(const FD3D12Payload* Payload, const ID3D12Fence* Fence, uint64 Value, ED3D12FenceOpType OpType)
+{
+	if (!IsEnabled(ERenderTracingChannels::RHISubmission))
+	{
+		return;
+	}
+
+	RENDER_TRACE_LOG(D3D12ManualFenceOp)
+		<< RENDER_TRACE_VALUE(D3D12ManualFenceOp, PayloadID, reinterpret_cast<uintptr_t>(Payload))
+		<< RENDER_TRACE_VALUE(D3D12ManualFenceOp, FenceID, reinterpret_cast<uintptr_t>(Fence))
+		<< RENDER_TRACE_VALUE(D3D12ManualFenceOp, FenceVal, Value)
+		<< RENDER_TRACE_VALUE(D3D12ManualFenceOp, Op, OpType);
+}
+
+void D3D12WaitManualFence(const FD3D12Payload* Payload, const ID3D12Fence* Fence, uint64 Value)
+{
+	D3D12ManualFenceOp(Payload, Fence, Value, D3D12FenceOp_Wait);
+}
+
+void D3D12SignalManualFence(const FD3D12Payload* Payload, const ID3D12Fence* Fence, uint64 Value)
+{
+	D3D12ManualFenceOp(Payload, Fence, Value, D3D12FenceOp_Signal);
+}
+
+void D3D12ProcessSubmissionQueueEnter()
+{
+	if (!IsEnabled(ERenderTracingChannels::RHISubmission))
+	{
+		return;
+	}
+
+	RENDER_TRACE_LOG(D3D12ProcessSubmissionQueueEnter);
+}
+
+void D3D12SubmitYieldOnUnresolvedSyncPoint(const FD3D12Queue& Queue, const FD3D12Payload* Payload, const FD3D12SyncPoint* SyncPoint)
+{
+	if (!IsEnabled(ERenderTracingChannels::RHISubmission))
+	{
+		return;
+	}
+
+	RENDER_TRACE_LOG(D3D12SubmitYieldOnUnresolvedSyncPoint)
+		<< RENDER_TRACE_VALUE(D3D12SubmitYieldOnUnresolvedSyncPoint, Queue, GetRenderTraceQueueIndex(Queue))
+		<< RENDER_TRACE_VALUE(D3D12SubmitYieldOnUnresolvedSyncPoint, PayloadID, reinterpret_cast<uintptr_t>(Payload))
+		<< RENDER_TRACE_VALUE(D3D12SubmitYieldOnUnresolvedSyncPoint, SyncPointID, reinterpret_cast<uintptr_t>(SyncPoint));
+}
+
+void D3D12SubmitYieldOnUnresolvedManualFence(const FD3D12Queue& Queue, const FD3D12Payload* Payload, const ID3D12Fence* Fence, uint64 Value)
+{
+	if (!IsEnabled(ERenderTracingChannels::RHISubmission))
+	{
+		return;
+	}
+
+	RENDER_TRACE_LOG(D3D12SubmitYieldOnUnresolvedManualFence)
+		<< RENDER_TRACE_VALUE(D3D12SubmitYieldOnUnresolvedManualFence, Queue, GetRenderTraceQueueIndex(Queue))
+		<< RENDER_TRACE_VALUE(D3D12SubmitYieldOnUnresolvedManualFence, PayloadID, reinterpret_cast<uintptr_t>(Payload))
+		<< RENDER_TRACE_VALUE(D3D12SubmitYieldOnUnresolvedManualFence, FenceID, reinterpret_cast<uintptr_t>(Fence))
+		<< RENDER_TRACE_VALUE(D3D12SubmitYieldOnUnresolvedManualFence, FenceValue, Value);
+}
+
+void D3D12SubmitWaitQueueFence(const FD3D12Payload* Payload, const FD3D12Queue& ExecutingQueue, const FD3D12Queue& WaitOnQueue, uint64 Value)
+{
+	if (!IsEnabled(ERenderTracingChannels::RHISubmission))
+	{
+		return;
+	}
+
+	RENDER_TRACE_LOG(D3D12SubmitWaitQueueFence)
+		<< RENDER_TRACE_VALUE(D3D12SubmitWaitQueueFence, PayloadID, reinterpret_cast<uintptr_t>(Payload))
+		<< RENDER_TRACE_VALUE(D3D12SubmitWaitQueueFence, ExecutingQueue, GetRenderTraceQueueIndex(ExecutingQueue))
+		<< RENDER_TRACE_VALUE(D3D12SubmitWaitQueueFence, WaitOnQueue, GetRenderTraceQueueIndex(WaitOnQueue))
+		<< RENDER_TRACE_VALUE(D3D12SubmitWaitQueueFence, Value, Value);
+}
+
+void D3D12SubmitWaitManualFence(const FD3D12Payload* Payload, const FD3D12Queue& ExecutingQueue, const ID3D12Fence* Fence, uint64 Value)
+{
+	if (!IsEnabled(ERenderTracingChannels::RHISubmission))
+	{
+		return;
+	}
+
+	RENDER_TRACE_LOG(D3D12SubmitWaitManualFence)
+		<< RENDER_TRACE_VALUE(D3D12SubmitWaitManualFence, PayloadID, reinterpret_cast<uintptr_t>(Payload))
+		<< RENDER_TRACE_VALUE(D3D12SubmitWaitManualFence, ExecutingQueue, GetRenderTraceQueueIndex(ExecutingQueue))
+		<< RENDER_TRACE_VALUE(D3D12SubmitWaitManualFence, FenceID, reinterpret_cast<uintptr_t>(Fence))
+		<< RENDER_TRACE_VALUE(D3D12SubmitWaitManualFence, Value, Value);
+}
+
+void D3D12SubmitPlatformCommandLists(const FD3D12Queue& Queue, TConstArrayView<FD3D12Payload*> Payloads, TConstArrayView<ID3D12CommandList*> CmdLists)
+{
+	if (!IsEnabled(ERenderTracingChannels::RHISubmission))
+	{
+		return;
+	}
+
+	RENDER_TRACE_LOG(D3D12SubmitPlatformCommandLists)
+		<< RENDER_TRACE_VALUE(D3D12SubmitPlatformCommandLists, Queue, GetRenderTraceQueueIndex(Queue))
+		<< RENDER_TRACE_ARRAY_VALUE(D3D12SubmitPlatformCommandLists, PayloadIDs, Payloads, uint64)
+		<< RENDER_TRACE_ARRAY_VALUE(D3D12SubmitPlatformCommandLists, CmdListIDs, CmdLists, uint64);
+}
+
+void D3D12SubmitSignalManualFence(const FD3D12Payload* Payload, const FD3D12Queue& Queue, const ID3D12Fence* Fence, uint64 Value)
+{
+	if (!IsEnabled(ERenderTracingChannels::RHISubmission))
+	{
+		return;
+	}
+
+	RENDER_TRACE_LOG(D3D12SubmitSignalManualFence)
+		<< RENDER_TRACE_VALUE(D3D12SubmitSignalManualFence, PayloadID, reinterpret_cast<uintptr_t>(Payload))
+		<< RENDER_TRACE_VALUE(D3D12SubmitSignalManualFence, Queue, GetRenderTraceQueueIndex(Queue))
+		<< RENDER_TRACE_VALUE(D3D12SubmitSignalManualFence, FenceID, reinterpret_cast<uintptr_t>(Fence))
+		<< RENDER_TRACE_VALUE(D3D12SubmitSignalManualFence, Value, Value);
+}
+
+void D3D12SubmitSignalQueueFence(const FD3D12Payload* Payload, const FD3D12Queue& Queue, uint64 Value)
+{
+	if (!IsEnabled(ERenderTracingChannels::RHISubmission))
+	{
+		return;
+	}
+
+	RENDER_TRACE_LOG(D3D12SubmitSignalQueueFence)
+		<< RENDER_TRACE_VALUE(D3D12SubmitSignalQueueFence, PayloadID, reinterpret_cast<uintptr_t>(Payload))
+		<< RENDER_TRACE_VALUE(D3D12SubmitSignalQueueFence, Queue, GetRenderTraceQueueIndex(Queue))
+		<< RENDER_TRACE_VALUE(D3D12SubmitSignalQueueFence, Value, Value);
+}
+
+void D3D12SubmitResolveSyncPoint(const FD3D12Queue& Queue, const FD3D12Payload* Payload, const FD3D12SyncPoint* SyncPoint, uint64 Value)
+{
+	if (!IsEnabled(ERenderTracingChannels::RHISubmission))
+	{
+		return;
+	}
+
+	RENDER_TRACE_LOG(D3D12SubmitResolveSyncPoint)
+		<< RENDER_TRACE_VALUE(D3D12SubmitResolveSyncPoint, Queue, GetRenderTraceQueueIndex(Queue))
+		<< RENDER_TRACE_VALUE(D3D12SubmitResolveSyncPoint, PayloadID, reinterpret_cast<uintptr_t>(Payload))
+		<< RENDER_TRACE_VALUE(D3D12SubmitResolveSyncPoint, SyncPointID, reinterpret_cast<uintptr_t>(SyncPoint))
+		<< RENDER_TRACE_VALUE(D3D12SubmitResolveSyncPoint, Value, Value);
+}
+
+void D3D12ProcessSubmissionQueueExit(uint8 Status)
+{
+	if (!IsEnabled(ERenderTracingChannels::RHISubmission))
+	{
+		return;
+	}
+
+	RENDER_TRACE_LOG(D3D12ProcessSubmissionQueueExit)
+		<< RENDER_TRACE_VALUE(D3D12ProcessSubmissionQueueExit, Status, Status);
+}
+
+void D3D12ProcessInterruptQueueEnter()
+{
+	if (!IsEnabled(ERenderTracingChannels::RHISubmission))
+	{
+		return;
+	}
+
+	RENDER_TRACE_LOG(D3D12ProcessInterruptQueueEnter);
+}
+
+void D3D12InterruptQueueFenceSignaled(const FD3D12Queue& Queue, FD3D12Payload* PendingPayload, uint64 CurrentFenceValue, uint64 LastCPUSignaledFenceValue)
+{
+	if (!IsEnabled(ERenderTracingChannels::RHISubmission))
+	{
+		return;
+	}
+
+	RENDER_TRACE_LOG(D3D12InterruptQueueFenceSignaled)
+		<< RENDER_TRACE_VALUE(D3D12InterruptQueueFenceSignaled, Queue, GetRenderTraceQueueIndex(Queue))
+		<< RENDER_TRACE_VALUE(D3D12InterruptQueueFenceSignaled, PendingPayloadID, reinterpret_cast<uintptr_t>(PendingPayload))
+		<< RENDER_TRACE_VALUE(D3D12InterruptQueueFenceSignaled, CurrentFenceValue, CurrentFenceValue)
+		<< RENDER_TRACE_VALUE(D3D12InterruptQueueFenceSignaled, LastCPUSignaledFenceValue, LastCPUSignaledFenceValue);
+}
+
+void D3D12ProcessInterruptQueueExit(uint8 Status)
+{
+	if (!IsEnabled(ERenderTracingChannels::RHISubmission))
+	{
+		return;
+	}
+
+	RENDER_TRACE_LOG(D3D12ProcessInterruptQueueExit)
+		<< RENDER_TRACE_VALUE(D3D12ProcessInterruptQueueExit, Status, Status);
+}
+
+}
+#endif
diff --git a/Engine/Source/Runtime/D3D12RHI/Private/D3D12Tracing.h b/Engine/Source/Runtime/D3D12RHI/Private/D3D12Tracing.h
new file mode 100644
--- /dev/null
+++ b/Engine/Source/Runtime/D3D12RHI/Private/D3D12Tracing.h
@@ -0,0 +1,47 @@
+// Copyright Epic Games, Inc. All Rights Reserved.
+
+#pragma once
+
+#include "RenderTracing.h"
+
+#if UE_RENDER_TRACING_ENABLED
+
+class IRHIComputeContext;
+class FD3D12ContextCommon;
+class FD3D12Queue;
+struct FD3D12Payload;
+class FD3D12SyncPoint;
+struct ID3D12CommandList;
+struct ID3D12Fence;
+
+namespace RenderTracing
+{
+
+void D3D12AssociateContext(const FD3D12ContextCommon* D3D12Context, const IRHIComputeContext* RHIContext);
+void D3D12PayloadCreated(const FD3D12Payload* Payload, const FD3D12Queue& Queue);
+void D3D12AddPayload(const FD3D12ContextCommon* D3D12Context, const FD3D12Payload* Payload);
+void D3D12OpenCommandList(const FD3D12Payload* Payload, const ID3D12CommandList* D3D12CmdList);
+void D3D12CloseCommandList(const FD3D12Payload* Payload, const ID3D12CommandList* D3D12CmdList);
+void D3D12WaitSyncPoint(const FD3D12Payload* Payload, const FD3D12SyncPoint* SyncPoint);
+void D3D12SignalSyncPoint(const FD3D12Payload* Payload, const FD3D12SyncPoint* SyncPoint);
+void D3D12WaitManualFence(const FD3D12Payload* Payload, const ID3D12Fence* Fence, uint64 Value);
+void D3D12SignalManualFence(const FD3D12Payload* Payload, const ID3D12Fence* Fence, uint64 Value);
+
+void D3D12ProcessSubmissionQueueEnter();
+void D3D12SubmitYieldOnUnresolvedSyncPoint(const FD3D12Queue& Queue, const FD3D12Payload* Payload, const FD3D12SyncPoint* SyncPoint);
+void D3D12SubmitYieldOnUnresolvedManualFence(const FD3D12Queue& Queue, const FD3D12Payload* Payload, const ID3D12Fence* Fence, uint64 Value);
+void D3D12SubmitWaitQueueFence(const FD3D12Payload* Payload, const FD3D12Queue& ExecutingQueue, const FD3D12Queue& WaitOnQueue, uint64 Value);
+void D3D12SubmitWaitManualFence(const FD3D12Payload* Payload, const FD3D12Queue& ExecutingQueue, const ID3D12Fence* Fence, uint64 Value);
+void D3D12SubmitPlatformCommandLists(const FD3D12Queue& Queue, TConstArrayView<FD3D12Payload*> Payloads, TConstArrayView<ID3D12CommandList*> CmdLists);
+void D3D12SubmitSignalManualFence(const FD3D12Payload* Payload, const FD3D12Queue& Queue, const ID3D12Fence* Fence, uint64 Value);
+void D3D12SubmitSignalQueueFence(const FD3D12Payload* Payload, const FD3D12Queue& Queue, uint64 Value);
+void D3D12SubmitResolveSyncPoint(const FD3D12Queue& Queue, const FD3D12Payload* Payload, const FD3D12SyncPoint* SyncPoint, uint64 Value);
+void D3D12ProcessSubmissionQueueExit(uint8 Status);
+
+void D3D12ProcessInterruptQueueEnter();
+void D3D12InterruptQueueFenceSignaled(const FD3D12Queue& Queue, FD3D12Payload* PendingPayload, uint64 CurrentFenceValue, uint64 LastCPUSignaledFenceValue);
+void D3D12ProcessInterruptQueueExit(uint8 Status);
+
+}
+
+#endif
diff --git a/Engine/Source/Runtime/D3D12RHI/Private/Windows/WindowsD3D12Submission.h b/Engine/Source/Runtime/D3D12RHI/Private/Windows/WindowsD3D12Submission.h
--- a/Engine/Source/Runtime/D3D12RHI/Private/Windows/WindowsD3D12Submission.h
+++ b/Engine/Source/Runtime/D3D12RHI/Private/Windows/WindowsD3D12Submission.h
@@ -8,11 +8,14 @@
 #include "D3D12Submission.h" // Circular include
 #endif
 
+#include "D3D12Tracing.h"
+
 // Standard Windows implementation, used to mark the type as 'final'.
 struct FD3D12Payload final : public FD3D12PayloadBase
 {
 	FD3D12Payload(FD3D12Queue& Queue)
 		: FD3D12PayloadBase(Queue)
 	{
+		IF_RENDER_TRACING_ENABLED(RenderTracing::D3D12PayloadCreated(this, Queue));
 	}
 };
diff --git a/Engine/Source/Runtime/RHI/Private/RHICommandList.cpp b/Engine/Source/Runtime/RHI/Private/RHICommandList.cpp
--- a/Engine/Source/Runtime/RHI/Private/RHICommandList.cpp
+++ b/Engine/Source/Runtime/RHI/Private/RHICommandList.cpp
@@ -21,6 +21,7 @@
 #include "Stats/ThreadIdleStats.h"
 #include "HAL/PlatformMisc.h"
 #include "HAL/LowLevelMemStats.h"
+#include "RHITracing.h"
 
 CSV_DEFINE_CATEGORY_MODULE(RHI_API, RHITStalls, false);
 CSV_DEFINE_CATEGORY_MODULE(RHI_API, RHITFlushes, false);
@@ -155,6 +156,7 @@ FRHICommandListBase::FRHICommandListBase(FPersistentState const& InPersistentSta
 {
 	DispatchEvent->SetDebugName(TEXT("FRHICommandListBase::DispatchEvent"));
 	CommandLink = &Root;
+	IF_RENDER_TRACING_ENABLED(RenderTracing::CommandListCreated(*this, InPersistentState.bImmediate));
 }
 
 FRHICommandListBase::~FRHICommandListBase()
@@ -172,6 +174,8 @@ FRHICommandListBase::~FRHICommandListBase()
 	}
 
 #endif
+
+	IF_RENDER_TRACING_ENABLED(RenderTracing::CommandListDestroyed(*this));
 }
 
 void FRHICommandListBase::InsertParallelRenderPass_Base(TSharedPtr<FRHIParallelRenderPassInfo> const& InInfo, TArray<FRHISubCommandList*>&& SubCommandLists)
@@ -197,6 +201,8 @@ void FRHICommandListBase::InsertParallelRenderPass_Base(TSharedPtr<FRHIParallelR
 	{
 		FRHICommandListBase* HeapCmdList = new FRHICommandListBase(MoveTemp(*this));
 
+		IF_RENDER_TRACING_ENABLED(RenderTracing::CommandListDetach(*this, *HeapCmdList, false));
+
 		(this)->~FRHICommandListBase();
 		new (this) FRHICommandListBase(HeapCmdList->PersistentState);
 
@@ -253,6 +259,8 @@ void FRHICommandListBase::FinishRecording()
 {
 	checkf(IsImmediate() || PersistentState.CurrentFenceScope == nullptr, TEXT("Finished recording with an open RHI fence scope."));
 
+	IF_RENDER_TRACING_ENABLED(RenderTracing::CommandListFinishRecording(*this, bUsesLockFence));
+
 	if (PersistentState.CurrentFenceScope && PersistentState.CurrentFenceScope->bFenceRequested)
 	{
 		PersistentState.CurrentFenceScope->bFenceRequested = false;
@@ -379,6 +387,8 @@ void FRHICommandListBase::ActivatePipelines(ERHIPipeline Pipelines)
 			ExecutingCmdList.ComputeContext = nullptr;
 		}
 
+		IF_RENDER_TRACING_ENABLED(RenderTracing::ActivatePipelines(ExecutingCmdList, static_cast<uint8>(NewPipelines)));
+
 		//
 		// Grab the appropriate command contexts from the RHI if we don't already have them.
 		//
@@ -407,6 +417,7 @@ void FRHICommandListBase::ActivatePipelines(ERHIPipeline Pipelines)
 							? GDynamicRHI->RHIGetCommandContext(Pipeline)
 							: GDynamicRHI->RHIGetDefaultContext();
 					}
+					IF_RENDER_TRACING_ENABLED(RenderTracing::SetTranslateContext(ExecutingCmdList, Context, static_cast<uint8>(Pipeline)));
 				}
 
 				if (bSinglePipeline)
@@ -425,6 +436,7 @@ void FRHICommandListBase::ActivatePipelines(ERHIPipeline Pipelines)
 				{
 					Context = GDynamicRHI->RHIGetCommandContext(Pipeline);
 					check(Context);
+					IF_RENDER_TRACING_ENABLED(RenderTracing::SetTranslateContext(ExecutingCmdList, Context, static_cast<uint8>(Pipeline)));
 				}
 
 				if (bSinglePipeline)
@@ -856,10 +868,12 @@ void FRHICommandListExecutor::FSubmitState::Dispatch_ProcessCommandList(FRHIComm
 			// Close the current translate chain
 			//
 
+			IF_RENDER_TRACING_ENABLED(RenderTracing::SplitTranslateJob(TranslateJobs.Current, *CmdList, bShouldSplitParallel, bShouldSplitForThreshold, bShouldSplitForParentChild));
+
 			bool bAddChildWait = TranslateJobs.Current->bUsingSubCmdLists && GRHIParallelRHIExecuteChildWait;
 			bool bAddParentWait = !TranslateJobs.Current->bUsingSubCmdLists && CmdList->IsSubCommandList() && GRHIParallelRHIExecuteParentWait;
 
-			FGraphEventRef Event = Dispatch_CloseTranslateChain(false);
+			FGraphEventRef Event = Dispatch_CloseTranslateChain(false, 0);
 			if (bAddChildWait)
 
 			{
@@ -936,6 +950,8 @@ void FRHICommandListExecutor::FSubmitState::Dispatch_ProcessCommandList(FRHIComm
 		Prereqs.Add(GRHICommandList.LastClose);
 	}
 
+	IF_RENDER_TRACING_ENABLED(RenderTracing::DispatchCommandList(TranslateJobs.Current, *CmdList));
+
 	TranslateJobs.Current->EnqueueTranslateTask(MoveTemp(Prereqs),
 		[TranslateJob = TranslateJobs.Current, CmdList](FTaskPipe::ITask&)
 		{
@@ -943,7 +959,7 @@ void FRHICommandListExecutor::FSubmitState::Dispatch_ProcessCommandList(FRHIComm
 		});
 }
 
-FGraphEventRef FRHICommandListExecutor::FSubmitState::Dispatch_CloseTranslateChain(bool bLastChainInSubmission)
+FGraphEventRef FRHICommandListExecutor::FSubmitState::Dispatch_CloseTranslateChain(bool bLastChainInSubmission, uint32 RTSubmissionID)
 {
 	// Tell the current chain there will be no more command lists. This finalizes the RHI contexts.
 	FTaskPipe* TranslatePipe = TranslateJobs.Current->EnqueueTranslateTask({},
@@ -1001,6 +1017,17 @@ FGraphEventRef FRHICommandListExecutor::FSubmitState::Dispatch_CloseTranslateCha
 		}
 	}
 
+#if UE_RENDER_TRACING_ENABLED
+	// Pointers won't have the bottom bit set, so we use this to distinguish between submissions instigated by the render
+	// thread via ERHISubmitFlags::SubmitToGPU and eager submissions caused by a translate chain being closed. In the latter
+	// case we want to link the submission to the last translate job in the chain.
+	uint64 TracingSubmissionID = !bLastChainInSubmission
+		? static_cast<uint64>(reinterpret_cast<uintptr_t>(TranslateJobs.Current))
+		: (static_cast<uint64>(RTSubmissionID) << 32) | 0x01;
+#else
+	uint64 TracingSubmissionID = 0;
+#endif
+
 	TranslateJobs.Previous = TranslateJobs.Current;
 	TranslateJobs.Current = nullptr;
 
@@ -1017,21 +1044,21 @@ FGraphEventRef FRHICommandListExecutor::FSubmitState::Dispatch_CloseTranslateCha
 
 		// Start a new task to process the submission queue from 'First'
 		GRHICommandList.EnqueueSubmitTask({ TranslateEvent },
-			[this, FirstJob = TranslateJobs.First](FTaskPipe::ITask& Task)
+			[this, FirstJob = TranslateJobs.First, TracingSubmissionID](FTaskPipe::ITask& Task)
 			{
-				Submit_SubmitToGPU(FirstJob, Task);
+				Submit_SubmitToGPU(FirstJob, Task, TracingSubmissionID);
 			});
 	}
 
 	return TranslateEvent;
 }
 
-void FRHICommandListExecutor::FSubmitState::Dispatch_CloseSubmission(TOptional<int64> LocalStatsFrame)
+void FRHICommandListExecutor::FSubmitState::Dispatch_CloseSubmission(TOptional<int64> LocalStatsFrame, uint32 RTSubmissionID)
 {
 	SCOPED_NAMED_EVENT(RHI_Dispatch_CloseSubmission, FColor::Purple);
 
 	// We still have a translate chain open. Close it.
-	Dispatch_CloseTranslateChain(true);
+	Dispatch_CloseTranslateChain(true, RTSubmissionID);
 
 	// Forward the close command to the submit pipe.
 	FSubmitState::FCloseArgs CloseArgs;
@@ -1149,6 +1176,8 @@ void FRHICommandListExecutor::FTranslateState::Translate_ExecuteCommandList(FRHI
 {
 	SCOPED_NAMED_EVENT(RHI_Translate, FColor::Orange);
 
+	IF_RENDER_TRACING_ENABLED(RenderTracing::BeginTranslateCommandList(this, *CmdList));
+
 	if (CmdList->ParallelRenderPassEnd)
 	{
 		// Retrieve the contexts from the previous parent cmdlist.
@@ -1229,6 +1258,8 @@ void FRHICommandListExecutor::FTranslateState::Translate_ExecuteCommandList(FRHI
 		CmdList->ParallelRenderPassBegin->Contexts = CmdList->Contexts;
 	}
 
+	IF_RENDER_TRACING_ENABLED(RenderTracing::EndTranslateCommandList(*CmdList));
+
 	delete CmdList;
 }
 
@@ -1268,6 +1299,8 @@ void FRHICommandListExecutor::FTranslateState::Translate_CloseChain()
 				}
 			}
 		}
+
+		IF_RENDER_TRACING_ENABLED(RenderTracing::FinalizeTranslateJob(this, bParallel, bUsingSubCmdLists));
 	}
 }
 
@@ -1275,7 +1308,7 @@ void FRHICommandListExecutor::FTranslateState::Translate_CloseChain()
 // -------------                                  Submit Pipe                                  -------------
 // ---------------------------------------------------------------------------------------------------------
 
-void FRHICommandListExecutor::FSubmitState::Submit_SubmitToGPU(FTranslateState* First, FTaskPipe::ITask& Task)
+void FRHICommandListExecutor::FSubmitState::Submit_SubmitToGPU(FTranslateState* First, FTaskPipe::ITask& Task, uint64 TracingSubmissionID)
 {
 	SCOPED_NAMED_EVENT(RHI_SubmitToGPU, FColor::Purple);
 
@@ -1283,6 +1316,10 @@ void FRHICommandListExecutor::FSubmitState::Submit_SubmitToGPU(FTranslateState*
 	// taking as many from the submission queue as we can.
 	TArray<IRHIPlatformCommandList*> CommandLists;
 
+#if UE_RENDER_TRACING_ENABLED
+	TArray<void*, TInlineAllocator<32>> TranslateJobIDs;
+#endif
+
 	FTranslateState* Current = First;
 	while (true)
 	{
@@ -1293,7 +1330,9 @@ void FRHICommandListExecutor::FSubmitState::Submit_SubmitToGPU(FTranslateState*
 			Task.ExtendTask({ Current->TranslateEvent },[this, First = Current](FTaskPipe::ITask& Task)
 			{
 				SCOPED_NAMED_EVENT(RHI_ExtendedSubmit, FColor::Orange);
-				Submit_SubmitToGPU(First, Task);
+				// We can compute this value even when tracing is disabled since it's free.
+				uint64 TracingSubmissionID = static_cast<uint64>(reinterpret_cast<uintptr_t>(First));
+				Submit_SubmitToGPU(First, Task, TracingSubmissionID);
 			});
 			break;
 		}
@@ -1338,6 +1377,10 @@ void FRHICommandListExecutor::FSubmitState::Submit_SubmitToGPU(FTranslateState*
 			}
 		}
 
+#if UE_RENDER_TRACING_ENABLED
+		TranslateJobIDs.Add(Current);
+#endif
+
 		// Try to close the queue. Use the current job's 'this' pointer as a sentinel
 		// to indicate to the dispatch pipe that the queue was atomically closed.
 		FTranslateState* Expected = nullptr;
@@ -1350,6 +1393,13 @@ void FRHICommandListExecutor::FSubmitState::Submit_SubmitToGPU(FTranslateState*
 		Current = Expected;
 	}
 
+#if UE_RENDER_TRACING_ENABLED
+	if (!TranslateJobIDs.IsEmpty())
+	{
+		RenderTracing::SubmitTranslateJobs(TranslateJobIDs, TracingSubmissionID);
+	}
+#endif
+
 	// Submit everything we gathered to the GPU.
 	if (CommandLists.Num())
 	{
@@ -1438,6 +1488,8 @@ RHI_API FGraphEventRef FRHICommandListExecutor::Submit(TConstArrayView<FRHIComma
 		// Move the contents of the immediate command list into a new heap-allocated instance.
 		ImmCmdList = new FRHICommandListBase(MoveTemp(static_cast<FRHICommandListBase&>(CommandListImmediate)));
 
+		IF_RENDER_TRACING_ENABLED(RenderTracing::CommandListDetach(CommandListImmediate, *ImmCmdList, true));
+
 		// Now reset the immediate command list.
 
 		// Destruct and reconstruct the base type in-place to reset all members to
@@ -1470,6 +1522,12 @@ RHI_API FGraphEventRef FRHICommandListExecutor::Submit(TConstArrayView<FRHIComma
 		ConsumeCmdList(ConsumeCmdList, CmdList);
 	}
 
+#if UE_RENDER_TRACING_ENABLED
+	const uint32 RTSubmissionID = RenderTracing::SubmitCommandLists(AllCmdLists, SubmitFlags);
+#else
+	const uint32 RTSubmissionID = 0;
+#endif
+
 	//
 	// Submission of RHI command lists involves a chain of "dispatch" tasks.
 	// These tasks wait for FinishRecording() to be called on each RHI command list, then start translate tasks to replay those command lists into RHI contexts.
@@ -1570,10 +1628,10 @@ RHI_API FGraphEventRef FRHICommandListExecutor::Submit(TConstArrayView<FRHIComma
 	#endif
 
 		EnqueueDispatchTask({},
-			[this, State = SubmitState, LocalStatsFrame](FTaskPipe::ITask&)
+			[this, State = SubmitState, LocalStatsFrame, RTSubmissionID](FTaskPipe::ITask&)
 			{
 				SCOPED_NAMED_EVENT(RHI_FinalizeSubmission, FColor::Yellow);
-				State->Dispatch_CloseSubmission(LocalStatsFrame);
+				State->Dispatch_CloseSubmission(LocalStatsFrame, RTSubmissionID);
 			});
 
 		SubmitState = nullptr;
diff --git a/Engine/Source/Runtime/RHI/Private/RHITracing.cpp b/Engine/Source/Runtime/RHI/Private/RHITracing.cpp
new file mode 100644
--- /dev/null
+++ b/Engine/Source/Runtime/RHI/Private/RHITracing.cpp
@@ -0,0 +1,263 @@
+// Copyright Epic Games, Inc. All Rights Reserved.
+
+#include "RHITracing.h"
+#include "Trace/Trace.inl"
+#include "HAL/PlatformTime.h"
+
+#if UE_RENDER_TRACING_ENABLED
+
+RENDER_TRACE_EVENT_BEGIN(CommandListCreated)
+	UE_TRACE_EVENT_FIELD(uint64, CmdListID)
+	UE_TRACE_EVENT_FIELD(bool, bImmediate)
+UE_TRACE_EVENT_END()
+
+RENDER_TRACE_EVENT_BEGIN(CommandListDetach)
+	UE_TRACE_EVENT_FIELD(uint64, SourceCmdListID)
+	UE_TRACE_EVENT_FIELD(uint64, DetachedCmdListID)
+	UE_TRACE_EVENT_FIELD(bool, bIsImmediateFlush)
+UE_TRACE_EVENT_END()
+
+RENDER_TRACE_EVENT_BEGIN(CommandListFinishRecording)
+	UE_TRACE_EVENT_FIELD(uint64, CmdListID)
+	UE_TRACE_EVENT_FIELD(uint8, Flags)
+UE_TRACE_EVENT_END()
+
+RENDER_TRACE_EVENT_BEGIN(CommandListDestroyed)
+	UE_TRACE_EVENT_FIELD(uint64, CmdListID)
+UE_TRACE_EVENT_END()
+
+RENDER_TRACE_EVENT_BEGIN(SubmitCommandLists)
+	UE_TRACE_EVENT_FIELD(uint64[], CmdListIDs)
+	UE_TRACE_EVENT_FIELD(uint32, RTSubmissionID)
+	UE_TRACE_EVENT_FIELD(uint16, Flags)
+UE_TRACE_EVENT_END()
+
+RENDER_TRACE_EVENT_BEGIN(TranslateJobCreated)
+	UE_TRACE_EVENT_FIELD(uint64, TranslateJobID)
+	UE_TRACE_EVENT_FIELD(uint8, Flags)
+UE_TRACE_EVENT_END()
+
+RENDER_TRACE_EVENT_BEGIN(SplitTranslateJob)
+	UE_TRACE_EVENT_FIELD(uint64, TranslateJobID)
+	UE_TRACE_EVENT_FIELD(uint64, CmdListID)
+	UE_TRACE_EVENT_FIELD(uint8, Flags)
+UE_TRACE_EVENT_END()
+
+RENDER_TRACE_EVENT_BEGIN(DispatchCommandList)
+	UE_TRACE_EVENT_FIELD(uint64, TranslateJobID)
+	UE_TRACE_EVENT_FIELD(uint64, CmdListID)
+UE_TRACE_EVENT_END()
+
+RENDER_TRACE_EVENT_BEGIN(BeginTranslateCommandList)
+	UE_TRACE_EVENT_FIELD(uint64, TranslateJobID)
+	UE_TRACE_EVENT_FIELD(uint64, CmdListID)
+UE_TRACE_EVENT_END()
+
+RENDER_TRACE_EVENT_BEGIN(ActivatePipelines)
+	UE_TRACE_EVENT_FIELD(uint64, CmdListID)
+	UE_TRACE_EVENT_FIELD(uint8, Pipelines)
+UE_TRACE_EVENT_END()
+
+RENDER_TRACE_EVENT_BEGIN(SetTranslateContext)
+	UE_TRACE_EVENT_FIELD(uint64, CmdListID)
+	UE_TRACE_EVENT_FIELD(uint64, ContextID)
+	UE_TRACE_EVENT_FIELD(uint8, Pipeline)
+UE_TRACE_EVENT_END()
+
+RENDER_TRACE_EVENT_BEGIN(EndTranslateCommandList)
+	UE_TRACE_EVENT_FIELD(uint64, CmdListID)
+UE_TRACE_EVENT_END()
+
+RENDER_TRACE_EVENT_BEGIN(FinalizeTranslateJob)
+	UE_TRACE_EVENT_FIELD(uint64, TranslateJobID)
+	UE_TRACE_EVENT_FIELD(uint8, Flags)
+UE_TRACE_EVENT_END()
+
+RENDER_TRACE_EVENT_BEGIN(SubmitTranslateJobs)
+	UE_TRACE_EVENT_FIELD(uint64[], TranslateJobIDs)
+	UE_TRACE_EVENT_FIELD(uint64, SubmissionID)
+UE_TRACE_EVENT_END()
+
+namespace RenderTracing
+{
+
+static uint32 GNextRTSubmissionID = 1;
+
+void CommandListCreated(const FRHICommandListBase& CmdList, bool bImmediate)
+{
+	if (!IsEnabled(ERenderTracingChannels::RHICmdLists))
+	{
+		return;
+	}
+
+	RENDER_TRACE_LOG(CommandListCreated)
+		<< RENDER_TRACE_VALUE(CommandListCreated, CmdListID, reinterpret_cast<uintptr_t>(&CmdList))
+		<< RENDER_TRACE_VALUE(CommandListCreated, bImmediate, bImmediate);
+}
+
+void CommandListDetach(const FRHICommandListBase& SourceCmdList, const FRHICommandListBase& DetachedCmdList, bool bIsImmediateFlush)
+{
+	if (!IsEnabled(ERenderTracingChannels::RHICmdLists))
+	{
+		return;
+	}
+
+	RENDER_TRACE_LOG(CommandListDetach)
+		<< RENDER_TRACE_VALUE(CommandListDetach, SourceCmdListID, reinterpret_cast<uintptr_t>(&SourceCmdList))
+		<< RENDER_TRACE_VALUE(CommandListDetach, DetachedCmdListID, reinterpret_cast<uintptr_t>(&DetachedCmdList))
+		<< RENDER_TRACE_VALUE(CommandListDetach, bIsImmediateFlush, bIsImmediateFlush);
+}
+
+void CommandListFinishRecording(const FRHICommandListBase& CmdList, bool bUsesLockFence)
+{
+	if (!IsEnabled(ERenderTracingChannels::RHICmdLists))
+	{
+		return;
+	}
+
+	uint8 Flags = 0;
+	Flags |= (uint8)bUsesLockFence << 0;
+
+	RENDER_TRACE_LOG(CommandListFinishRecording)
+		<< RENDER_TRACE_VALUE(CommandListFinishRecording, CmdListID, reinterpret_cast<uintptr_t>(&CmdList))
+		<< RENDER_TRACE_VALUE(CommandListFinishRecording, Flags, Flags);
+}
+
+void CommandListDestroyed(const FRHICommandListBase& CmdList)
+{
+	if (!IsEnabled(ERenderTracingChannels::RHICmdLists))
+	{
+		return;
+	}
+
+	RENDER_TRACE_LOG(CommandListDestroyed)
+		<< RENDER_TRACE_VALUE(CommandListDestroyed, CmdListID, reinterpret_cast<uintptr_t>(&CmdList));
+}
+
+uint32 SubmitCommandLists(TConstArrayView<FRHICommandListBase*> CmdLists, ERHISubmitFlags Flags)
+{
+	if (!IsEnabled(ERenderTracingChannels::RHITranslation))
+	{
+		return 0;
+	}
+
+	// If SubmitToGPU is set, we want to assign a unique ID to this event so that it can be linked later
+	// to the RHI submission task.
+	const uint32 RTSubmissionID = EnumHasAllFlags(Flags, ERHISubmitFlags::SubmitToGPU) ? GNextRTSubmissionID++ : 0;
+
+	RENDER_TRACE_LOG(SubmitCommandLists)
+		<< RENDER_TRACE_ARRAY_VALUE(SubmitCommandLists, CmdListIDs, CmdLists, uint64)
+		<< RENDER_TRACE_VALUE(SubmitCommandLists, RTSubmissionID, RTSubmissionID)
+		<< RENDER_TRACE_VALUE(SubmitCommandLists, Flags, static_cast<uint16>(Flags));
+
+	return RTSubmissionID;
+}
+
+void SplitTranslateJob(const void* TranslateJobID, FRHICommandListBase& CmdList, bool bSplitParallel, bool bSplitForThreshold, bool bSplitForParentChild)
+{
+	if (!IsEnabled(ERenderTracingChannels::RHITranslation))
+	{
+		return;
+	}
+
+	uint8 Flags = 0;
+	Flags |= (uint8)bSplitParallel << 0;
+	Flags |= (uint8)bSplitForThreshold << 1;
+	Flags |= (uint8)bSplitForParentChild << 2;
+
+	RENDER_TRACE_LOG(SplitTranslateJob)
+		<< RENDER_TRACE_VALUE(SplitTranslateJob, TranslateJobID, reinterpret_cast<uintptr_t>(TranslateJobID))
+		<< RENDER_TRACE_VALUE(SplitTranslateJob, CmdListID, reinterpret_cast<uintptr_t>(&CmdList))
+		<< RENDER_TRACE_VALUE(SplitTranslateJob, Flags, Flags);
+}
+
+void DispatchCommandList(const void* TranslateJobID, const FRHICommandListBase& CmdList)
+{
+	if (!IsEnabled(ERenderTracingChannels::RHITranslation))
+	{
+		return;
+	}
+
+	RENDER_TRACE_LOG(DispatchCommandList)
+		<< RENDER_TRACE_VALUE(DispatchCommandList, TranslateJobID, reinterpret_cast<uintptr_t>(TranslateJobID))
+		<< RENDER_TRACE_VALUE(DispatchCommandList, CmdListID, reinterpret_cast<uintptr_t>(&CmdList));
+}
+
+void BeginTranslateCommandList(const void* TranslateJobID, const FRHICommandListBase& CmdList)
+{
+	if (!IsEnabled(ERenderTracingChannels::RHITranslation))
+	{
+		return;
+	}
+
+	RENDER_TRACE_LOG(BeginTranslateCommandList)
+		<< RENDER_TRACE_VALUE(BeginTranslateCommandList, TranslateJobID, reinterpret_cast<uintptr_t>(TranslateJobID))
+		<< RENDER_TRACE_VALUE(BeginTranslateCommandList, CmdListID, reinterpret_cast<uintptr_t>(&CmdList));
+}
+
+void ActivatePipelines(const FRHICommandListBase& CmdList, uint8 Pipelines)
+{
+	if (!IsEnabled(ERenderTracingChannels::RHITranslation))
+	{
+		return;
+	}
+
+	RENDER_TRACE_LOG(ActivatePipelines)
+		<< RENDER_TRACE_VALUE(ActivatePipelines, CmdListID, reinterpret_cast<uintptr_t>(&CmdList))
+		<< RENDER_TRACE_VALUE(ActivatePipelines, Pipelines, Pipelines);
+}
+
+void SetTranslateContext(const FRHICommandListBase& CmdList, const IRHIComputeContext* Context, uint8 Pipeline)
+{
+	if (!IsEnabled(ERenderTracingChannels::RHITranslation))
+	{
+		return;
+	}
+
+	RENDER_TRACE_LOG(SetTranslateContext)
+		<< RENDER_TRACE_VALUE(SetTranslateContext, CmdListID, reinterpret_cast<uintptr_t>(&CmdList))
+		<< RENDER_TRACE_VALUE(SetTranslateContext, ContextID, reinterpret_cast<uintptr_t>(Context))
+		<< RENDER_TRACE_VALUE(SetTranslateContext, Pipeline, Pipeline);
+}
+
+void EndTranslateCommandList(const FRHICommandListBase& CmdList)
+{
+	if (!IsEnabled(ERenderTracingChannels::RHITranslation))
+	{
+		return;
+	}
+
+	RENDER_TRACE_LOG(EndTranslateCommandList)
+		<< RENDER_TRACE_VALUE(EndTranslateCommandList, CmdListID, reinterpret_cast<uintptr_t>(&CmdList));
+}
+
+void FinalizeTranslateJob(const void* TranslateJobID, bool bParallel, bool bUsingSubCmdLists)
+{
+	if (!IsEnabled(ERenderTracingChannels::RHITranslation))
+	{
+		return;
+	}
+
+	uint8 Flags = 0;
+	Flags |= (uint8)bParallel << 0;
+	Flags |= (uint8)bUsingSubCmdLists << 1;
+
+	RENDER_TRACE_LOG(FinalizeTranslateJob)
+		<< RENDER_TRACE_VALUE(FinalizeTranslateJob, TranslateJobID, reinterpret_cast<uintptr_t>(TranslateJobID))
+		<< RENDER_TRACE_VALUE(FinalizeTranslateJob, Flags, Flags);
+}
+
+void SubmitTranslateJobs(TConstArrayView<void*> TranslateJobIDs, uint64 SubmissionID)
+{
+	if (!IsEnabled(ERenderTracingChannels::RHITranslation))
+	{
+		return;
+	}
+
+	RENDER_TRACE_LOG(SubmitTranslateJobs)
+		<< RENDER_TRACE_ARRAY_VALUE(SubmitTranslateJobs, TranslateJobIDs, TranslateJobIDs, uint64)
+		<< RENDER_TRACE_VALUE(SubmitTranslateJobs, SubmissionID, SubmissionID);
+}
+
+}
+#endif
diff --git a/Engine/Source/Runtime/RHI/Private/RHITracing.h b/Engine/Source/Runtime/RHI/Private/RHITracing.h
new file mode 100644
--- /dev/null
+++ b/Engine/Source/Runtime/RHI/Private/RHITracing.h
@@ -0,0 +1,31 @@
+// Copyright Epic Games, Inc. All Rights Reserved.
+
+#pragma once
+
+#include "RenderTracing.h"
+#include "RHICommandList.h"
+#include "RHIContext.h"
+
+#if UE_RENDER_TRACING_ENABLED
+
+namespace RenderTracing
+{
+
+void CommandListCreated(const FRHICommandListBase& CmdList, bool bImmediate);
+void CommandListDetach(const FRHICommandListBase& SourceCmdList, const FRHICommandListBase& DetachedCmdList, bool bIsImmediateFlush);
+void CommandListFinishRecording(const FRHICommandListBase& CmdList, bool bUsesLockFence);
+void CommandListDestroyed(const FRHICommandListBase& CmdList);
+
+uint32 SubmitCommandLists(TConstArrayView<FRHICommandListBase*> CmdLists, ERHISubmitFlags Flags);
+void SplitTranslateJob(const void* TranslateJobID, FRHICommandListBase& CmdList, bool bSplitParallel, bool bSplitForThreshold, bool bSplitForParentChild);
+void DispatchCommandList(const void* TranslateJobID, const FRHICommandListBase& CmdList);
+void BeginTranslateCommandList(const void* TranslateJobID, const FRHICommandListBase& CmdList);
+void EndTranslateCommandList(const FRHICommandListBase& CmdList);
+void ActivatePipelines(const FRHICommandListBase& CmdList, uint8 Pipelines);
+void SetTranslateContext(const FRHICommandListBase& CmdList, const IRHIComputeContext* Context, uint8 Pipeline);
+void FinalizeTranslateJob(const void* TranslateJobID, bool bParallel, bool bUsingSubCmdLists);
+void SubmitTranslateJobs(TConstArrayView<void*> TranslateJobIDs, uint64 SubmissionID);
+
+}
+
+#endif
diff --git a/Engine/Source/Runtime/RHI/Private/RenderTracing.cpp b/Engine/Source/Runtime/RHI/Private/RenderTracing.cpp
new file mode 100644
--- /dev/null
+++ b/Engine/Source/Runtime/RHI/Private/RenderTracing.cpp
@@ -0,0 +1,23 @@
+// Copyright Epic Games, Inc. All Rights Reserved.
+
+#include "RenderTracing.h"
+#include "Trace/Trace.inl"
+#include "HAL/PlatformTime.h"
+
+UE_TRACE_CHANNEL_DEFINE(RenderTraceChannel);
+
+#if UE_RENDER_TRACING_ENABLED
+
+namespace RenderTracing
+{
+
+// TODO: make this configurable.
+static ERenderTracingChannels GRenderTracingChannels = ERenderTracingChannels::Submission;
+
+bool IsEnabled(ERenderTracingChannels ChannelMask)
+{
+	return EnumHasAnyFlags(GRenderTracingChannels, ChannelMask) && UE_TRACE_CHANNELEXPR_IS_ENABLED(RenderTraceChannel);
+}
+
+}
+#endif
diff --git a/Engine/Source/Runtime/RHI/Public/RHICommandList.h b/Engine/Source/Runtime/RHI/Public/RHICommandList.h
--- a/Engine/Source/Runtime/RHI/Public/RHICommandList.h
+++ b/Engine/Source/Runtime/RHI/Public/RHICommandList.h
@@ -5252,10 +5252,10 @@ private:
 
 		void Dispatch_ProcessCommandList(FRHICommandListBase* CmdList);
 		void Dispatch_PrepareBreadcrumbs(FRHICommandListBase* CmdList);
-		FGraphEventRef Dispatch_CloseTranslateChain(bool bLastChainInSubmission);
-		void Dispatch_CloseSubmission(TOptional<int64> LocalStatsFrame);
+		FGraphEventRef Dispatch_CloseTranslateChain(bool bLastChainInSubmission, uint32 RTSubmissionID);
+		void Dispatch_CloseSubmission(TOptional<int64> LocalStatsFrame, uint32 RTSubmissionID);
 
-		void Submit_SubmitToGPU(FTranslateState* First, FTaskPipe::ITask& Task);
+		void Submit_SubmitToGPU(FTranslateState* First, FTaskPipe::ITask& Task, uint64 TracingSubmissionID);
 		void Submit_Close(const FCloseArgs& Args);
 
 		FSubmitState(bool bBatchSubmissions)
diff --git a/Engine/Source/Runtime/RHI/Public/RenderTracing.h b/Engine/Source/Runtime/RHI/Public/RenderTracing.h
new file mode 100644
--- /dev/null
+++ b/Engine/Source/Runtime/RHI/Public/RenderTracing.h
@@ -0,0 +1,50 @@
+// Copyright Epic Games, Inc. All Rights Reserved.
+
+#pragma once
+
+#include "Trace/Trace.h"
+#include "Containers/ArrayView.h"
+#include "Misc/EnumClassFlags.h"
+
+#define UE_RENDER_TRACING_ENABLED UE_TRACE_ENABLED && !IS_PROGRAM && !UE_BUILD_SHIPPING
+
+#if UE_RENDER_TRACING_ENABLED
+	#define IF_RENDER_TRACING_ENABLED(Op) Op
+#else
+	#define IF_RENDER_TRACING_ENABLED(Op)
+#endif
+
+UE_TRACE_CHANNEL_EXTERN(RenderTraceChannel, RHI_API);
+
+#if UE_RENDER_TRACING_ENABLED
+
+#define RENDER_TRACE_EVENT_BEGIN(EventName, ...) UE_TRACE_EVENT_BEGIN(RenderTrace, EventName##Message, ##__VA_ARGS__) \
+	UE_TRACE_EVENT_FIELD(uint64, Timestamp)
+
+#define RENDER_TRACE_LOG(EventName) UE_TRACE_LOG(RenderTrace, EventName##Message, RenderTraceChannel) \
+	<< EventName##Message.Timestamp(FPlatformTime::Cycles64())
+
+#define RENDER_TRACE_VALUE(EventName, ValueName, Value) EventName##Message.ValueName(Value)
+#define RENDER_TRACE_ARRAY_VALUE(EventName, ValueName, Value, ElemType) EventName##Message.ValueName(reinterpret_cast<const ElemType*>(Value.GetData()), Value.Num())
+
+enum class ERenderTracingChannels : uint64
+{
+	None = 0,
+	RDGPasses = 1 << 0,
+	RHICmdLists = 1 << 1,
+	RHITranslation = 1 << 2,
+	RHISubmission = 1 << 3,
+
+	Submission = RDGPasses | RHICmdLists | RHITranslation | RHISubmission
+};
+ENUM_CLASS_FLAGS(ERenderTracingChannels);
+
+
+namespace RenderTracing
+{
+
+RHI_API bool IsEnabled(ERenderTracingChannels ChannelMask);
+
+}
+
+#endif
diff --git a/Engine/Source/Runtime/RenderCore/Private/RenderGraphBuilder.cpp b/Engine/Source/Runtime/RenderCore/Private/RenderGraphBuilder.cpp
--- a/Engine/Source/Runtime/RenderCore/Private/RenderGraphBuilder.cpp
+++ b/Engine/Source/Runtime/RenderCore/Private/RenderGraphBuilder.cpp
@@ -9,6 +9,7 @@
 #include "VisualizeTexture.h"
 #include "ProfilingDebugging/CsvProfiler.h"
 #include "Async/ParallelFor.h"
+#include "RenderTracing.h"
 
 struct FParallelPassSet : public FRHICommandListImmediate::FQueuedCommandList
 {
@@ -3079,6 +3080,9 @@ void FRDGBuilder::SetupDispatchPassExecute()
 			IF_RDG_ENABLE_DEBUG(UserValidation.ValidateExecutePassBegin(DispatchPass));
 
 			FRDGDispatchPassBuilder DispatchPassBuilder(DispatchPass);
+
+			IF_RENDER_TRACING_ENABLED(RenderTracing::BeginExecuteDispatchPass(*DispatchPass));
+
 			DispatchPass->LaunchDispatchPassTasks(DispatchPassBuilder);
 			DispatchPassBuilder.Finish();
 
@@ -3481,6 +3485,8 @@ void FRDGBuilder::ExecutePassEpilogue(FRHIComputeCommandList& RHICmdListPass, FR
 
 void FRDGBuilder::ExecutePass(FRHIComputeCommandList& RHICmdListPass, FRDGPass* Pass)
 {
+	IF_RENDER_TRACING_ENABLED(RenderTracing::BeginExecuteRDGPass(*Pass, RHICmdListPass));
+
 	// Note that we must do this before doing anything with RHICmdListPass.
 	// For example, if this pass only executes on GPU 1 we want to avoid adding a
 	// 0-duration event for this pass on GPU 0's time line.
@@ -3492,6 +3498,8 @@ void FRDGBuilder::ExecutePass(FRHIComputeCommandList& RHICmdListPass, FRDGPass*
 	Pass->Execute(RHICmdListPass);
 
 	ExecutePassEpilogue(RHICmdListPass, Pass);
+
+	IF_RENDER_TRACING_ENABLED(RenderTracing::EndExecuteRDGPass(*Pass));
 }
 
 void FRDGBuilder::ExecuteSerialPass(FRHIComputeCommandList& RHICmdListPass, FRDGPass* Pass)
diff --git a/Engine/Source/Runtime/RenderCore/Private/RenderGraphPass.cpp b/Engine/Source/Runtime/RenderCore/Private/RenderGraphPass.cpp
--- a/Engine/Source/Runtime/RenderCore/Private/RenderGraphPass.cpp
+++ b/Engine/Source/Runtime/RenderCore/Private/RenderGraphPass.cpp
@@ -3,6 +3,7 @@
 #include "RenderGraphPass.h"
 #include "RenderGraphPrivate.h"
 #include "RenderGraphBuilder.h"
+#include "RenderTracing.h"
 
 FEmptyShaderParameters FRDGSentinelPass::EmptyShaderParameters;
 
@@ -149,6 +150,8 @@ FRHICommandList* FRDGDispatchPassBuilder::CreateCommandList()
 	
 	RHICmdList->SetStaticUniformBuffers(StaticUniformBuffers);
 
+	IF_RENDER_TRACING_ENABLED(RenderTracing::AddDispatchPassCommandList(*Pass, *RHICmdList));
+
 	return RHICmdList;
 }
 
@@ -157,6 +160,7 @@ void FRDGDispatchPassBuilder::Finish()
 	// With serial execution the pass commands are embedded in the immediate command list instead.
 	if (!Pass->bParallelExecute)
 	{
+		IF_RENDER_TRACING_ENABLED(RenderTracing::EndExecuteDispatchPass(*Pass, nullptr));
 		Pass->CommandListsEvent.Trigger();
 		return;
 	}
@@ -197,6 +201,8 @@ void FRDGDispatchPassBuilder::Finish()
 	FRDGBuilder::ExecutePassEpilogue(*RHICmdList, Pass);
 	FRDGBuilder::PopPreScopes(*RHICmdList, Pass);
 
+	IF_RENDER_TRACING_ENABLED(RenderTracing::EndExecuteDispatchPass(*Pass, RHICmdList));
+
 	RHICmdList->FinishRecording();
 }
 
diff --git a/Engine/Source/Runtime/RenderCore/Private/RenderGraphTrace.cpp b/Engine/Source/Runtime/RenderCore/Private/RenderGraphTrace.cpp
--- a/Engine/Source/Runtime/RenderCore/Private/RenderGraphTrace.cpp
+++ b/Engine/Source/Runtime/RenderCore/Private/RenderGraphTrace.cpp
@@ -501,4 +501,103 @@ void FRDGTrace::AddBufferPassDependency(FRDGBuffer* Buffer, FRDGPass* Pass)
 	Buffer->TracePasses.Add(Pass->Handle);
 }
 
-#endif
\ No newline at end of file
+#endif // RDG_ENABLE_TRACE
+
+#if UE_RENDER_TRACING_ENABLED
+
+RENDER_TRACE_EVENT_BEGIN(BeginExecuteRDGPass)
+	UE_TRACE_EVENT_FIELD(uint64, PassID)
+	UE_TRACE_EVENT_FIELD(UE::Trace::WideString, Name)
+	UE_TRACE_EVENT_FIELD(uint64, CmdListID)
+	UE_TRACE_EVENT_FIELD(uint16, PassFlags)
+	UE_TRACE_EVENT_FIELD(uint8, TaskMode)
+UE_TRACE_EVENT_END()
+
+RENDER_TRACE_EVENT_BEGIN(EndExecuteRDGPass)
+	UE_TRACE_EVENT_FIELD(uint64, PassID)
+UE_TRACE_EVENT_END()
+
+RENDER_TRACE_EVENT_BEGIN(BeginExecuteDispatchPass)
+	UE_TRACE_EVENT_FIELD(uint64, PassID)
+	UE_TRACE_EVENT_FIELD(UE::Trace::WideString, Name)
+UE_TRACE_EVENT_END()
+
+RENDER_TRACE_EVENT_BEGIN(AddDispatchPassCommandList)
+	UE_TRACE_EVENT_FIELD(uint64, PassID)
+	UE_TRACE_EVENT_FIELD(uint64, CmdListID)
+UE_TRACE_EVENT_END()
+
+RENDER_TRACE_EVENT_BEGIN(EndExecuteDispatchPass)
+	UE_TRACE_EVENT_FIELD(uint64, PassID)
+	UE_TRACE_EVENT_FIELD(uint64, EpilogueCmdListID)
+UE_TRACE_EVENT_END()
+
+namespace RenderTracing
+{
+
+void BeginExecuteRDGPass(const FRDGPass& Pass, const FRHICommandListBase& CmdList)
+{
+	if (!IsEnabled(ERenderTracingChannels::RDGPasses))
+	{
+		return;
+	}
+
+	// TODO: don't send pass names as strings, find a more efficient representation.
+
+	RENDER_TRACE_LOG(BeginExecuteRDGPass)
+		<< RENDER_TRACE_VALUE(BeginExecuteRDGPass, PassID, reinterpret_cast<uintptr_t>(&Pass))
+		<< RENDER_TRACE_VALUE(BeginExecuteRDGPass, Name, Pass.GetName())
+		<< RENDER_TRACE_VALUE(BeginExecuteRDGPass, CmdListID, reinterpret_cast<uintptr_t>(&CmdList))
+		<< RENDER_TRACE_VALUE(BeginExecuteRDGPass, PassFlags, static_cast<uint16>(Pass.GetFlags()))
+		<< RENDER_TRACE_VALUE(BeginExecuteRDGPass, TaskMode, static_cast<uint8>(Pass.GetTaskMode()));
+}
+
+void EndExecuteRDGPass(const FRDGPass& Pass)
+{
+	if (!IsEnabled(ERenderTracingChannels::RDGPasses))
+	{
+		return;
+	}
+
+	RENDER_TRACE_LOG(EndExecuteRDGPass)
+		<< RENDER_TRACE_VALUE(EndExecuteRDGPass, PassID, reinterpret_cast<uintptr_t>(&Pass));
+}
+
+void BeginExecuteDispatchPass(const FRDGDispatchPass& Pass)
+{
+	if (!IsEnabled(ERenderTracingChannels::RDGPasses))
+	{
+		return;
+	}
+
+	RENDER_TRACE_LOG(BeginExecuteDispatchPass)
+		<< RENDER_TRACE_VALUE(BeginExecuteDispatchPass, PassID, reinterpret_cast<uintptr_t>(&Pass))
+		<< RENDER_TRACE_VALUE(BeginExecuteDispatchPass, Name, Pass.GetName());
+}
+
+void AddDispatchPassCommandList(const FRDGDispatchPass& Pass, const FRHICommandListBase& CmdList)
+{
+	if (!IsEnabled(ERenderTracingChannels::RDGPasses))
+	{
+		return;
+	}
+
+	RENDER_TRACE_LOG(AddDispatchPassCommandList)
+		<< RENDER_TRACE_VALUE(AddDispatchPassCommandList, PassID, reinterpret_cast<uintptr_t>(&Pass))
+		<< RENDER_TRACE_VALUE(AddDispatchPassCommandList, CmdListID, reinterpret_cast<uintptr_t>(&CmdList));
+}
+
+void EndExecuteDispatchPass(const FRDGDispatchPass& Pass, const FRHICommandListBase* EpilogueCmdList)
+{
+	if (!IsEnabled(ERenderTracingChannels::RDGPasses))
+	{
+		return;
+	}
+
+	RENDER_TRACE_LOG(EndExecuteDispatchPass)
+		<< RENDER_TRACE_VALUE(EndExecuteDispatchPass, PassID, reinterpret_cast<uintptr_t>(&Pass))
+		<< RENDER_TRACE_VALUE(EndExecuteDispatchPass, EpilogueCmdListID, reinterpret_cast<uintptr_t>(EpilogueCmdList));
+}
+
+}
+#endif
diff --git a/Engine/Source/Runtime/RenderCore/Public/RenderGraphTrace.h b/Engine/Source/Runtime/RenderCore/Public/RenderGraphTrace.h
--- a/Engine/Source/Runtime/RenderCore/Public/RenderGraphTrace.h
+++ b/Engine/Source/Runtime/RenderCore/Public/RenderGraphTrace.h
@@ -5,13 +5,17 @@
 #include "HAL/Platform.h"
 #include "RHITransientResourceAllocator.h"
 #include "RenderGraphDefinitions.h"
+#include "RenderGraphPass.h"
 #include "Trace/Trace.h"
+#include "RenderTracing.h"
 
 class FRDGBuffer;
 class FRDGBuilder;
 class FRDGPass;
+class FRDGDispatchPass;
 class FRDGTexture;
 class FRDGViewableResource;
+class FRHICommandListBase;
 namespace UE { namespace Trace { class FChannel; } }
 
 #if RDG_ENABLE_TRACE
@@ -40,4 +44,20 @@ private:
 	bool bEnabled;
 };
 
-#endif
+#endif // RDG_ENABLE_TRACE
+
+#if UE_RENDER_TRACING_ENABLED
+
+namespace RenderTracing
+{
+
+void BeginExecuteRDGPass(const FRDGPass& Pass, const FRHICommandListBase& CmdList);
+void EndExecuteRDGPass(const FRDGPass& Pass);
+void BeginExecuteDispatchPass(const FRDGDispatchPass& Pass);
+void AddDispatchPassCommandList(const FRDGDispatchPass& Pass, const FRHICommandListBase& CmdList);
+void EndExecuteDispatchPass(const FRDGDispatchPass& Pass, const FRHICommandListBase* EpilogueCmdList);
+
+}
+
+#endif // UE_RENDER_TRACING_ENABLED
+
