Use generic names for platform submission tracing.

#rnx

#Synced-CL 51937798

diff --git a/Engine/Plugins/RenderTraceInsights/Source/RenderTraceInsights/Private/RenderTraceAnalyzer.cpp b/Engine/Plugins/RenderTraceInsights/Source/RenderTraceInsights/Private/RenderTraceAnalyzer.cpp
--- a/Engine/Plugins/RenderTraceInsights/Source/RenderTraceInsights/Private/RenderTraceAnalyzer.cpp
+++ b/Engine/Plugins/RenderTraceInsights/Source/RenderTraceInsights/Private/RenderTraceAnalyzer.cpp
@@ -12,16 +12,16 @@ namespace UE
 namespace RenderTraceInsights
 {
 
-enum ED3D12CommandListOpType
+enum ECommandListOpType
 {
-	D3D12CommandListOp_Open = 0,
-	D3D12CommandListOp_Close = 1
+	CommandListOp_Open = 0,
+	CommandListOp_Close = 1
 };
 
-enum ED3D12FenceOpType
+enum EFenceOpType
 {
-	D3D12FenceOp_Wait = 0,
-	D3D12FenceOp_Signal = 1
+	FenceOp_Wait = 0,
+	FenceOp_Signal = 1
 };
 
 FRenderTraceAnalyzer::FRenderTraceAnalyzer(TraceServices::IAnalysisSession& InSession, IEditableRenderTraceProvider& InProvider)
@@ -53,25 +53,25 @@ void FRenderTraceAnalyzer::OnAnalysisBegin(const FOnAnalysisContext& Context)
 	ROUTE_EVENT(EndTranslateCommandList);
 	ROUTE_EVENT(FinalizeTranslateJob);
 	ROUTE_EVENT(SubmitTranslateJobs);
-	ROUTE_EVENT(D3D12AssociateContext);
-	ROUTE_EVENT(D3D12PayloadCreated);
-	ROUTE_EVENT(D3D12AddPayload);
-	ROUTE_EVENT(D3D12CommandListOp);
-	ROUTE_EVENT(D3D12SyncPointOp);
-	ROUTE_EVENT(D3D12ManualFenceOp);
-	ROUTE_EVENT(D3D12ProcessSubmissionQueueEnter);
-	ROUTE_EVENT(D3D12SubmitYieldOnUnresolvedSyncPoint);
-	ROUTE_EVENT(D3D12SubmitYieldOnUnresolvedManualFence);
-	ROUTE_EVENT(D3D12SubmitWaitQueueFence);
-	ROUTE_EVENT(D3D12SubmitWaitManualFence);
-	ROUTE_EVENT(D3D12SubmitPlatformCommandLists);
-	ROUTE_EVENT(D3D12SubmitSignalManualFence);
-	ROUTE_EVENT(D3D12SubmitSignalQueueFence);
-	ROUTE_EVENT(D3D12SubmitResolveSyncPoint);
-	ROUTE_EVENT(D3D12ProcessSubmissionQueueExit);
-	ROUTE_EVENT(D3D12ProcessInterruptQueueEnter);
-	ROUTE_EVENT(D3D12InterruptQueueFenceSignaled);
-	ROUTE_EVENT(D3D12ProcessInterruptQueueExit);
+	ROUTE_EVENT(AssociateContext);
+	ROUTE_EVENT(PayloadCreated);
+	ROUTE_EVENT(AddPayload);
+	ROUTE_EVENT(CommandListOp);
+	ROUTE_EVENT(SyncPointOp);
+	ROUTE_EVENT(ManualFenceOp);
+	ROUTE_EVENT(ProcessSubmissionQueueEnter);
+	ROUTE_EVENT(SubmitYieldOnUnresolvedSyncPoint);
+	ROUTE_EVENT(SubmitYieldOnUnresolvedManualFence);
+	ROUTE_EVENT(SubmitWaitQueueFence);
+	ROUTE_EVENT(SubmitWaitManualFence);
+	ROUTE_EVENT(SubmitPlatformCommandLists);
+	ROUTE_EVENT(SubmitSignalManualFence);
+	ROUTE_EVENT(SubmitSignalQueueFence);
+	ROUTE_EVENT(SubmitResolveSyncPoint);
+	ROUTE_EVENT(ProcessSubmissionQueueExit);
+	ROUTE_EVENT(ProcessInterruptQueueEnter);
+	ROUTE_EVENT(InterruptQueueFenceSignaled);
+	ROUTE_EVENT(ProcessInterruptQueueExit);
 
 #undef ROUTE_EVENT
 }
@@ -817,7 +817,7 @@ void FRenderTraceAnalyzer::ProcessSetTranslateContext(const FEventData& EventDat
 	// with the context (such as signal batched sync points) and we still need to find the context -> job association when that
 	// happens. The association will still be valid when it does that, because these extra actions happen before returning the context
 	// to the pool, so another job can't steal the context while we still need the old association.
-	// Example flow with d3d12:
+	// Example flow with D3D12:
 	//		FRHICommandListExecutor::FTranslateState::Translate_CloseChain()
 	//			|-> GDynamicRHI->RHICloseTranslateChain()
 	//			|	|-> FD3D12DynamicRHI::RHIFinalizeContext()
@@ -825,7 +825,7 @@ void FRenderTraceAnalyzer::ProcessSetTranslateContext(const FEventData& EventDat
 	//			|		|	|-> Payload->SyncPointsToWait.Append(BatchedSyncPoints.ToWait) - this needs to find the current context -> translate job mapping
 	//			|		|-> CmdContext->GetParentDevice()->ReleaseContext() - this releases the context, so another thread can pick it up for its translate job
 	// ... another thread sends a SetTranslateContext packet and steals the context for a different job ...
-	//			|-> FRenderTracing::FinalizeTranslateJob() - we should only manipulate the context if it's still bound to the job
+	//			|-> RenderTracing::FinalizeTranslateJob() - we should only manipulate the context if it's still bound to the job
 	TranslateContext.RHIContextID = ContextID;
 	LiveRHIContexts.Add(ContextID, TranslateJobID);
 }
@@ -1022,7 +1022,7 @@ void FRenderTraceAnalyzer::ProcessSubmitTranslateJobs(const FEventData& EventDat
 	Timeline.AppendEndEvent(SubmitTask.EndTime);
 }
 
-void FRenderTraceAnalyzer::ProcessD3D12AssociateContext(const FEventData& EventData, uint32 ThreadID, double Timestamp)
+void FRenderTraceAnalyzer::ProcessAssociateContext(const FEventData& EventData, uint32 ThreadID, double Timestamp)
 {
 	// Since platform contexts are pooled, we will only get this message once, when the pooled object is created. These tend to live forever,
 	// so PlatformToRHIContext is not like other mapping tables, which are constantly updated as objects are created and destroyed. Its contents
@@ -1045,7 +1045,7 @@ void FRenderTraceAnalyzer::ProcessD3D12AssociateContext(const FEventData& EventD
 	PlatformToRHIContext.Add(PlatformContextID, RHIContextID);
 }
 
-void FRenderTraceAnalyzer::ProcessD3D12PayloadCreated(const FEventData& EventData, uint32 ThreadID, double Timestamp)
+void FRenderTraceAnalyzer::ProcessPayloadCreated(const FEventData& EventData, uint32 ThreadID, double Timestamp)
 {
 	const uint64 PayloadID = EventData.GetValue<uint64>("PayloadID");
 	const uint8 Queue = EventData.GetValue<uint8>("Queue");
@@ -1061,7 +1061,7 @@ void FRenderTraceAnalyzer::ProcessD3D12PayloadCreated(const FEventData& EventDat
 	LivePlatformPayloads.Add(PayloadID, NewPayloadEntry.Key);
 }
 
-void FRenderTraceAnalyzer::ProcessD3D12AddPayload(const FEventData& EventData, uint32 ThreadID, double Timestamp)
+void FRenderTraceAnalyzer::ProcessAddPayload(const FEventData& EventData, uint32 ThreadID, double Timestamp)
 {
 	const uint64 PlatformContextID = EventData.GetValue<uint64>("PlatformContextID");
 	const uint64 PayloadID = EventData.GetValue<uint64>("PayloadID");
@@ -1115,10 +1115,10 @@ void FRenderTraceAnalyzer::ProcessD3D12AddPayload(const FEventData& EventData, u
 	Payload.TranslateEvent.ItemID = TranslateContext.PayloadIDs.Num() - 1;
 }
 
-void FRenderTraceAnalyzer::ProcessD3D12CommandListOp(const FEventData& EventData, uint32 ThreadID, double Timestamp)
+void FRenderTraceAnalyzer::ProcessCommandListOp(const FEventData& EventData, uint32 ThreadID, double Timestamp)
 {
 	const uint64 PayloadID = EventData.GetValue<uint64>("PayloadID");
-	const uint64 D3D12CmdListID = EventData.GetValue<uint64>("D3D12CmdListID");
+	const uint64 PlatformCmdListID = EventData.GetValue<uint64>("PlatformCmdListID");
 	const uint8 Op = EventData.GetValue<uint8>("Op");
 
 	FPlatformPayload* Payload = FindPlatformPayload(PayloadID, TEXT("command list op"), ThreadID);
@@ -1129,13 +1129,13 @@ void FRenderTraceAnalyzer::ProcessD3D12CommandListOp(const FEventData& EventData
 
 	switch (Op)
 	{
-	case D3D12CommandListOp_Open:
+	case CommandListOp_Open:
 	{
-		Payload->CmdLists.Emplace(D3D12CmdListID);
+		Payload->CmdLists.Emplace(PlatformCmdListID);
 		break;
 	}
 
-	case D3D12CommandListOp_Close:
+	case CommandListOp_Close:
 	{
 		// Move the end of the payload to the end of this command list, but don't exceed the translate job lifetime if it's already set.
 		Payload->EndTime = FMath::Max(Timestamp, Payload->EndTime);
@@ -1153,13 +1153,13 @@ void FRenderTraceAnalyzer::ProcessD3D12CommandListOp(const FEventData& EventData
 	default:
 	{
 		UE_LOG(LogRenderTrace, Warning, TEXT("Got command list open event for payload ID %llx, command list ID %llx with unknown op code %u."),
-			PayloadID, D3D12CmdListID, Op);
+			PayloadID, PlatformCmdListID, Op);
 		break;
 	}
 	}
 }
 
-void FRenderTraceAnalyzer::ProcessD3D12SyncPointOp(const FEventData& EventData, uint32 ThreadID, double Timestamp)
+void FRenderTraceAnalyzer::ProcessSyncPointOp(const FEventData& EventData, uint32 ThreadID, double Timestamp)
 {
 	const uint64 PayloadID = EventData.GetValue<uint64>("PayloadID");
 	const uint64 SyncPointID = EventData.GetValue<uint64>("SyncPointID");
@@ -1189,7 +1189,7 @@ void FRenderTraceAnalyzer::ProcessD3D12SyncPointOp(const FEventData& EventData,
 	SyncPointList.Add(PermanentSyncPointID);
 }
 
-void FRenderTraceAnalyzer::ProcessD3D12ManualFenceOp(const FEventData& EventData, uint32 ThreadID, double Timestamp)
+void FRenderTraceAnalyzer::ProcessManualFenceOp(const FEventData& EventData, uint32 ThreadID, double Timestamp)
 {
 	const uint64 PayloadID = EventData.GetValue<uint64>("PayloadID");
 	const uint64 FenceID = EventData.GetValue<uint64>("FenceID");
@@ -1222,7 +1222,7 @@ void FRenderTraceAnalyzer::ProcessD3D12ManualFenceOp(const FEventData& EventData
 	SyncPointList.Add(PermanentSyncPointID);
 }
 
-void FRenderTraceAnalyzer::ProcessD3D12ProcessSubmissionQueueEnter(const FEventData& EventData, uint32 ThreadID, double Timestamp)
+void FRenderTraceAnalyzer::ProcessProcessSubmissionQueueEnter(const FEventData& EventData, uint32 ThreadID, double Timestamp)
 {
 	if (PendingSubmissionBatchID != INVALID_EVENT_ID)
 	{
@@ -1246,7 +1246,7 @@ void FRenderTraceAnalyzer::ProcessD3D12ProcessSubmissionQueueEnter(const FEventD
 	FSubmissionEvent& SubmissionEvent = Batch.Events.Emplace_GetRef(Timestamp, Queue);\
 	SubmissionEvent.Data.Emplace<FSubmissionEvent::F##Name>(__VA_ARGS__)
 	
-void FRenderTraceAnalyzer::ProcessD3D12SubmitYieldOnUnresolvedSyncPoint(const FEventData& EventData, uint32 ThreadID, double Timestamp)
+void FRenderTraceAnalyzer::ProcessSubmitYieldOnUnresolvedSyncPoint(const FEventData& EventData, uint32 ThreadID, double Timestamp)
 {
 	const uint8 Queue = EventData.GetValue<uint8>("Queue");
 	uint64 PayloadID = EventData.GetValue<uint64>("PayloadID");
@@ -1255,7 +1255,7 @@ void FRenderTraceAnalyzer::ProcessD3D12SubmitYieldOnUnresolvedSyncPoint(const FE
 	ADD_SUBMIT_EVENT(YieldSyncPoint, PayloadID, SyncPointID);
 }
 
-void FRenderTraceAnalyzer::ProcessD3D12SubmitYieldOnUnresolvedManualFence(const FEventData& EventData, uint32 ThreadID, double Timestamp)
+void FRenderTraceAnalyzer::ProcessSubmitYieldOnUnresolvedManualFence(const FEventData& EventData, uint32 ThreadID, double Timestamp)
 {
 	const uint8 Queue = EventData.GetValue<uint8>("Queue");
 	uint64 PayloadID = EventData.GetValue<uint64>("PayloadID");
@@ -1265,7 +1265,7 @@ void FRenderTraceAnalyzer::ProcessD3D12SubmitYieldOnUnresolvedManualFence(const
 	ADD_SUBMIT_EVENT(YieldManualFence, PayloadID, FenceID, FenceValue);
 }
 
-void FRenderTraceAnalyzer::ProcessD3D12SubmitWaitQueueFence(const FEventData& EventData, uint32 ThreadID, double Timestamp)
+void FRenderTraceAnalyzer::ProcessSubmitWaitQueueFence(const FEventData& EventData, uint32 ThreadID, double Timestamp)
 {
 	const uint64 PayloadID = EventData.GetValue<uint64>("PayloadID");
 	const uint8 Queue = EventData.GetValue<uint8>("ExecutingQueue");
@@ -1275,7 +1275,7 @@ void FRenderTraceAnalyzer::ProcessD3D12SubmitWaitQueueFence(const FEventData& Ev
 	ADD_SUBMIT_EVENT(WaitQueueFence, PayloadID, WaitOnQueue, Value);
 }
 
-void FRenderTraceAnalyzer::ProcessD3D12SubmitWaitManualFence(const FEventData& EventData, uint32 ThreadID, double Timestamp)
+void FRenderTraceAnalyzer::ProcessSubmitWaitManualFence(const FEventData& EventData, uint32 ThreadID, double Timestamp)
 {
 	const uint64 PayloadID = EventData.GetValue<uint64>("PayloadID");
 	const uint8 Queue = EventData.GetValue<uint8>("ExecutingQueue");
@@ -1285,11 +1285,11 @@ void FRenderTraceAnalyzer::ProcessD3D12SubmitWaitManualFence(const FEventData& E
 	ADD_SUBMIT_EVENT(WaitManualFence, PayloadID, FenceID, Value);
 }
 
-void FRenderTraceAnalyzer::ProcessD3D12SubmitPlatformCommandLists(const FEventData& EventData, uint32 ThreadID, double Timestamp)
+void FRenderTraceAnalyzer::ProcessSubmitPlatformCommandLists(const FEventData& EventData, uint32 ThreadID, double Timestamp)
 {
 	const uint8 Queue = EventData.GetValue<uint8>("Queue");
 	TConstArrayView<uint64> PayloadIDs = EventData.GetArrayView<uint64>("PayloadIDs");
-	TConstArrayView<uint64> CmdListIDs = EventData.GetArrayView<uint64>("CmdListIDs");
+	TConstArrayView<uint64> CmdListIDs = EventData.GetArrayView<uint64>("PlatformCmdListIDs");
 
 	TArray<uint32> PayloadPermanentIDs;
 	PayloadPermanentIDs.Reserve(PayloadIDs.Num());
@@ -1312,7 +1312,7 @@ void FRenderTraceAnalyzer::ProcessD3D12SubmitPlatformCommandLists(const FEventDa
 	}
 }
 
-void FRenderTraceAnalyzer::ProcessD3D12SubmitSignalManualFence(const FEventData& EventData, uint32 ThreadID, double Timestamp)
+void FRenderTraceAnalyzer::ProcessSubmitSignalManualFence(const FEventData& EventData, uint32 ThreadID, double Timestamp)
 {
 	const uint64 PayloadID = EventData.GetValue<uint64>("PayloadID");
 	const uint8 Queue = EventData.GetValue<uint8>("Queue");
@@ -1322,7 +1322,7 @@ void FRenderTraceAnalyzer::ProcessD3D12SubmitSignalManualFence(const FEventData&
 	ADD_SUBMIT_EVENT(SignalManualFence, PayloadID, FenceID, Value);
 }
 
-void FRenderTraceAnalyzer::ProcessD3D12SubmitSignalQueueFence(const FEventData& EventData, uint32 ThreadID, double Timestamp)
+void FRenderTraceAnalyzer::ProcessSubmitSignalQueueFence(const FEventData& EventData, uint32 ThreadID, double Timestamp)
 {
 	const uint64 PayloadID = EventData.GetValue<uint64>("PayloadID");
 	const uint8 Queue = EventData.GetValue<uint8>("Queue");
@@ -1331,7 +1331,7 @@ void FRenderTraceAnalyzer::ProcessD3D12SubmitSignalQueueFence(const FEventData&
 	ADD_SUBMIT_EVENT(SignalQueueFence, PayloadID, Value);
 }
 
-void FRenderTraceAnalyzer::ProcessD3D12SubmitResolveSyncPoint(const FEventData& EventData, uint32 ThreadID, double Timestamp)
+void FRenderTraceAnalyzer::ProcessSubmitResolveSyncPoint(const FEventData& EventData, uint32 ThreadID, double Timestamp)
 {
 	const uint8 Queue = EventData.GetValue<uint8>("Queue");
 	const uint64 PayloadID = EventData.GetValue<uint64>("PayloadID");
@@ -1388,7 +1388,7 @@ void FRenderTraceAnalyzer::EndPendingSubmissionBatch(double Timestamp, uint8 Exi
 	PendingSubmissionBatchID = INVALID_EVENT_ID;
 }
 
-void FRenderTraceAnalyzer::ProcessD3D12ProcessSubmissionQueueExit(const FEventData& EventData, uint32 ThreadID, double Timestamp)
+void FRenderTraceAnalyzer::ProcessProcessSubmissionQueueExit(const FEventData& EventData, uint32 ThreadID, double Timestamp)
 {
 	if (PendingSubmissionBatchID == INVALID_EVENT_ID)
 	{
@@ -1402,7 +1402,7 @@ void FRenderTraceAnalyzer::ProcessD3D12ProcessSubmissionQueueExit(const FEventDa
 
 #undef ADD_SUBMIT_EVENT
 
-void FRenderTraceAnalyzer::ProcessD3D12ProcessInterruptQueueEnter(const FEventData& EventData, uint32 ThreadID, double Timestamp)
+void FRenderTraceAnalyzer::ProcessProcessInterruptQueueEnter(const FEventData& EventData, uint32 ThreadID, double Timestamp)
 {
 	if (PendingInterruptWakeUpID != INVALID_EVENT_ID)
 	{
@@ -1414,7 +1414,7 @@ void FRenderTraceAnalyzer::ProcessD3D12ProcessInterruptQueueEnter(const FEventDa
 	PendingInterruptWakeUpID = NewWakeUp.Key;
 }
 
-void FRenderTraceAnalyzer::ProcessD3D12InterruptQueueFenceSignaled(const FEventData& EventData, uint32 ThreadID, double Timestamp)
+void FRenderTraceAnalyzer::ProcessInterruptQueueFenceSignaled(const FEventData& EventData, uint32 ThreadID, double Timestamp)
 {
 	const uint8 Queue = EventData.GetValue<uint8>("Queue");
 	const uint64 PendingPayloadID = EventData.GetValue<uint64>("PendingPayloadID");
@@ -1467,7 +1467,7 @@ void FRenderTraceAnalyzer::EndPendingInterruptWakeUp(double Timestamp, uint8 Exi
 	PendingInterruptWakeUpID = INVALID_EVENT_ID;
 }
 
-void FRenderTraceAnalyzer::ProcessD3D12ProcessInterruptQueueExit(const FEventData& EventData, uint32 ThreadID, double Timestamp)
+void FRenderTraceAnalyzer::ProcessProcessInterruptQueueExit(const FEventData& EventData, uint32 ThreadID, double Timestamp)
 {
 	if (PendingInterruptWakeUpID == INVALID_EVENT_ID)
 	{
@@ -1584,25 +1584,25 @@ bool FRenderTraceAnalyzer::OnEvent(uint16 RouteId, EStyle Style, const FOnEventC
 	DISPATCH_EVENT(EndTranslateCommandList);
 	DISPATCH_EVENT(FinalizeTranslateJob);
 	DISPATCH_EVENT(SubmitTranslateJobs);
-	DISPATCH_EVENT(D3D12AssociateContext);
-	DISPATCH_EVENT(D3D12PayloadCreated);
-	DISPATCH_EVENT(D3D12AddPayload);
-	DISPATCH_EVENT(D3D12CommandListOp);
-	DISPATCH_EVENT(D3D12SyncPointOp);
-	DISPATCH_EVENT(D3D12ManualFenceOp);
-	DISPATCH_EVENT(D3D12ProcessSubmissionQueueEnter);
-	DISPATCH_EVENT(D3D12SubmitYieldOnUnresolvedSyncPoint);
-	DISPATCH_EVENT(D3D12SubmitYieldOnUnresolvedManualFence);
-	DISPATCH_EVENT(D3D12SubmitWaitQueueFence);
-	DISPATCH_EVENT(D3D12SubmitWaitManualFence);
-	DISPATCH_EVENT(D3D12SubmitPlatformCommandLists);
-	DISPATCH_EVENT(D3D12SubmitSignalManualFence);
-	DISPATCH_EVENT(D3D12SubmitSignalQueueFence);
-	DISPATCH_EVENT(D3D12SubmitResolveSyncPoint);
-	DISPATCH_EVENT(D3D12ProcessSubmissionQueueExit);
-	DISPATCH_EVENT(D3D12ProcessInterruptQueueEnter);
-	DISPATCH_EVENT(D3D12InterruptQueueFenceSignaled);
-	DISPATCH_EVENT(D3D12ProcessInterruptQueueExit);
+	DISPATCH_EVENT(AssociateContext);
+	DISPATCH_EVENT(PayloadCreated);
+	DISPATCH_EVENT(AddPayload);
+	DISPATCH_EVENT(CommandListOp);
+	DISPATCH_EVENT(SyncPointOp);
+	DISPATCH_EVENT(ManualFenceOp);
+	DISPATCH_EVENT(ProcessSubmissionQueueEnter);
+	DISPATCH_EVENT(SubmitYieldOnUnresolvedSyncPoint);
+	DISPATCH_EVENT(SubmitYieldOnUnresolvedManualFence);
+	DISPATCH_EVENT(SubmitWaitQueueFence);
+	DISPATCH_EVENT(SubmitWaitManualFence);
+	DISPATCH_EVENT(SubmitPlatformCommandLists);
+	DISPATCH_EVENT(SubmitSignalManualFence);
+	DISPATCH_EVENT(SubmitSignalQueueFence);
+	DISPATCH_EVENT(SubmitResolveSyncPoint);
+	DISPATCH_EVENT(ProcessSubmissionQueueExit);
+	DISPATCH_EVENT(ProcessInterruptQueueEnter);
+	DISPATCH_EVENT(InterruptQueueFenceSignaled);
+	DISPATCH_EVENT(ProcessInterruptQueueExit);
 
 	case Packet_BeginExecuteRDGPass: ProcessBeginExecuteRDGPass(EventData, ThreadID, Timestamp, false); break;
 	case Packet_BeginExecuteDispatchPass: ProcessBeginExecuteRDGPass(EventData, ThreadID, Timestamp, true); break;
diff --git a/Engine/Plugins/RenderTraceInsights/Source/RenderTraceInsights/Private/RenderTraceAnalyzer.h b/Engine/Plugins/RenderTraceInsights/Source/RenderTraceInsights/Private/RenderTraceAnalyzer.h
--- a/Engine/Plugins/RenderTraceInsights/Source/RenderTraceInsights/Private/RenderTraceAnalyzer.h
+++ b/Engine/Plugins/RenderTraceInsights/Source/RenderTraceInsights/Private/RenderTraceAnalyzer.h
@@ -42,25 +42,25 @@ private:
 		Packet_EndTranslateCommandList,
 		Packet_FinalizeTranslateJob,
 		Packet_SubmitTranslateJobs,
-		Packet_D3D12AssociateContext,
-		Packet_D3D12PayloadCreated,
-		Packet_D3D12AddPayload,
-		Packet_D3D12CommandListOp,
-		Packet_D3D12SyncPointOp,
-		Packet_D3D12ManualFenceOp,
-		Packet_D3D12ProcessSubmissionQueueEnter,
-		Packet_D3D12SubmitYieldOnUnresolvedSyncPoint,
-		Packet_D3D12SubmitYieldOnUnresolvedManualFence,
-		Packet_D3D12SubmitWaitQueueFence,
-		Packet_D3D12SubmitWaitManualFence,
-		Packet_D3D12SubmitPlatformCommandLists,
-		Packet_D3D12SubmitSignalManualFence,
-		Packet_D3D12SubmitSignalQueueFence,
-		Packet_D3D12SubmitResolveSyncPoint,
-		Packet_D3D12ProcessSubmissionQueueExit,
-		Packet_D3D12ProcessInterruptQueueEnter,
-		Packet_D3D12InterruptQueueFenceSignaled,
-		Packet_D3D12ProcessInterruptQueueExit,
+		Packet_AssociateContext,
+		Packet_PayloadCreated,
+		Packet_AddPayload,
+		Packet_CommandListOp,
+		Packet_SyncPointOp,
+		Packet_ManualFenceOp,
+		Packet_ProcessSubmissionQueueEnter,
+		Packet_SubmitYieldOnUnresolvedSyncPoint,
+		Packet_SubmitYieldOnUnresolvedManualFence,
+		Packet_SubmitWaitQueueFence,
+		Packet_SubmitWaitManualFence,
+		Packet_SubmitPlatformCommandLists,
+		Packet_SubmitSignalManualFence,
+		Packet_SubmitSignalQueueFence,
+		Packet_SubmitResolveSyncPoint,
+		Packet_ProcessSubmissionQueueExit,
+		Packet_ProcessInterruptQueueEnter,
+		Packet_InterruptQueueFenceSignaled,
+		Packet_ProcessInterruptQueueExit,
 	};
 
 	struct FLiveTranslateJob
@@ -117,25 +117,25 @@ private:
 	DECLARE_EVENT_PROCESSING_FUNC(EndTranslateCommandList);
 	DECLARE_EVENT_PROCESSING_FUNC(FinalizeTranslateJob);
 	DECLARE_EVENT_PROCESSING_FUNC(SubmitTranslateJobs);
-	DECLARE_EVENT_PROCESSING_FUNC(D3D12AssociateContext);
-	DECLARE_EVENT_PROCESSING_FUNC(D3D12PayloadCreated);
-	DECLARE_EVENT_PROCESSING_FUNC(D3D12AddPayload);
-	DECLARE_EVENT_PROCESSING_FUNC(D3D12CommandListOp);
-	DECLARE_EVENT_PROCESSING_FUNC(D3D12SyncPointOp);
-	DECLARE_EVENT_PROCESSING_FUNC(D3D12ManualFenceOp);
-	DECLARE_EVENT_PROCESSING_FUNC(D3D12ProcessSubmissionQueueEnter);
-	DECLARE_EVENT_PROCESSING_FUNC(D3D12SubmitYieldOnUnresolvedSyncPoint);
-	DECLARE_EVENT_PROCESSING_FUNC(D3D12SubmitYieldOnUnresolvedManualFence);
-	DECLARE_EVENT_PROCESSING_FUNC(D3D12SubmitWaitQueueFence);
-	DECLARE_EVENT_PROCESSING_FUNC(D3D12SubmitWaitManualFence);
-	DECLARE_EVENT_PROCESSING_FUNC(D3D12SubmitPlatformCommandLists);
-	DECLARE_EVENT_PROCESSING_FUNC(D3D12SubmitSignalManualFence);
-	DECLARE_EVENT_PROCESSING_FUNC(D3D12SubmitSignalQueueFence);
-	DECLARE_EVENT_PROCESSING_FUNC(D3D12SubmitResolveSyncPoint);
-	DECLARE_EVENT_PROCESSING_FUNC(D3D12ProcessSubmissionQueueExit);
-	DECLARE_EVENT_PROCESSING_FUNC(D3D12ProcessInterruptQueueEnter);
-	DECLARE_EVENT_PROCESSING_FUNC(D3D12InterruptQueueFenceSignaled);
-	DECLARE_EVENT_PROCESSING_FUNC(D3D12ProcessInterruptQueueExit);
+	DECLARE_EVENT_PROCESSING_FUNC(AssociateContext);
+	DECLARE_EVENT_PROCESSING_FUNC(PayloadCreated);
+	DECLARE_EVENT_PROCESSING_FUNC(AddPayload);
+	DECLARE_EVENT_PROCESSING_FUNC(CommandListOp);
+	DECLARE_EVENT_PROCESSING_FUNC(SyncPointOp);
+	DECLARE_EVENT_PROCESSING_FUNC(ManualFenceOp);
+	DECLARE_EVENT_PROCESSING_FUNC(ProcessSubmissionQueueEnter);
+	DECLARE_EVENT_PROCESSING_FUNC(SubmitYieldOnUnresolvedSyncPoint);
+	DECLARE_EVENT_PROCESSING_FUNC(SubmitYieldOnUnresolvedManualFence);
+	DECLARE_EVENT_PROCESSING_FUNC(SubmitWaitQueueFence);
+	DECLARE_EVENT_PROCESSING_FUNC(SubmitWaitManualFence);
+	DECLARE_EVENT_PROCESSING_FUNC(SubmitPlatformCommandLists);
+	DECLARE_EVENT_PROCESSING_FUNC(SubmitSignalManualFence);
+	DECLARE_EVENT_PROCESSING_FUNC(SubmitSignalQueueFence);
+	DECLARE_EVENT_PROCESSING_FUNC(SubmitResolveSyncPoint);
+	DECLARE_EVENT_PROCESSING_FUNC(ProcessSubmissionQueueExit);
+	DECLARE_EVENT_PROCESSING_FUNC(ProcessInterruptQueueEnter);
+	DECLARE_EVENT_PROCESSING_FUNC(InterruptQueueFenceSignaled);
+	DECLARE_EVENT_PROCESSING_FUNC(ProcessInterruptQueueExit);
 #undef DECLARE_EVENT_PROCESSING_FUNC
 
 	// These are shared between regular and dispatch RDG passes so they have different signatures from the other event processing methods.
diff --git a/Engine/Source/Runtime/D3D12RHI/Private/D3D12CommandContext.cpp b/Engine/Source/Runtime/D3D12RHI/Private/D3D12CommandContext.cpp
--- a/Engine/Source/Runtime/D3D12RHI/Private/D3D12CommandContext.cpp
+++ b/Engine/Source/Runtime/D3D12RHI/Private/D3D12CommandContext.cpp
@@ -75,7 +75,7 @@ FD3D12CommandContext::FD3D12CommandContext(FD3D12Device* InParent, ED3D12QueueTy
 	// translate jobs with IRHIComputeContext pointers, and FD3D12ContextCommon does not derive from that interface. Instead,
 	// this class derives from both, so this is the place where we can map FD3D12ContextCommon pointers to IRHIComputeContext
 	// pointers, to allow the tracing code to reconstruct the translate job -> RHI context -> D3D12 context dependency chains.
-	IF_RENDER_TRACING_ENABLED(RenderTracing::D3D12AssociateContext(this, this));
+	IF_RENDER_TRACING_ENABLED(RenderTracing::D3D12::AssociateContext(this, this));
 }
 
 FD3D12CommandContext::~FD3D12CommandContext()
@@ -246,7 +246,7 @@ void FD3D12ContextCommon::WaitSyncPoint(FD3D12SyncPoint* SyncPoint)
 	FD3D12Payload* Payload = GetPayload(EPhase::Wait);
 	Payload->SyncPointsToWait.Add(SyncPoint);
 
-	IF_RENDER_TRACING_ENABLED(RenderTracing::D3D12WaitSyncPoint(Payload, SyncPoint));
+	IF_RENDER_TRACING_ENABLED(RenderTracing::D3D12::WaitSyncPoint(Payload, SyncPoint));
 }
 
 void FD3D12ContextCommon::SignalSyncPoint(FD3D12SyncPoint* SyncPoint)
@@ -259,7 +259,7 @@ void FD3D12ContextCommon::SignalSyncPoint(FD3D12SyncPoint* SyncPoint)
 	FD3D12Payload* Payload = GetPayload(EPhase::Signal);
 	Payload->SyncPointsToSignal.Add(SyncPoint);
 
-	IF_RENDER_TRACING_ENABLED(RenderTracing::D3D12SignalSyncPoint(Payload, SyncPoint));
+	IF_RENDER_TRACING_ENABLED(RenderTracing::D3D12::SignalSyncPoint(Payload, SyncPoint));
 }
 
 void FD3D12ContextCommon::SignalManualFence(ID3D12Fence* Fence, uint64 Value)
@@ -272,7 +272,7 @@ void FD3D12ContextCommon::SignalManualFence(ID3D12Fence* Fence, uint64 Value)
 	FD3D12Payload* Payload = GetPayload(EPhase::Signal);
 	Payload->ManualFencesToSignal.Emplace(Fence, Value, nullptr);
 
-	IF_RENDER_TRACING_ENABLED(RenderTracing::D3D12SignalManualFence(Payload, Fence, Value));
+	IF_RENDER_TRACING_ENABLED(RenderTracing::D3D12::SignalManualFence(Payload, Fence, Value));
 }
 
 void FD3D12ContextCommon::WaitManualFence(ID3D12Fence* Fence, uint64 Value, const TSharedPtr<std::atomic<bool>>& bHasSignalled)
@@ -285,7 +285,7 @@ void FD3D12ContextCommon::WaitManualFence(ID3D12Fence* Fence, uint64 Value, cons
 	FD3D12Payload* Payload = GetPayload(EPhase::Wait);
 	Payload->ManualFencesToWait.Emplace(Fence, Value, bHasSignalled);
 
-	IF_RENDER_TRACING_ENABLED(RenderTracing::D3D12WaitManualFence(Payload, Fence, Value));
+	IF_RENDER_TRACING_ENABLED(RenderTracing::D3D12::WaitManualFence(Payload, Fence, Value));
 }
 
 FD3D12QueryLocation FD3D12ContextCommon::AllocateQuery(ED3D12QueryType Type, void* Target)
@@ -361,7 +361,7 @@ void FD3D12ContextCommon::OpenCommandList()
 	FD3D12Payload* Payload = GetPayload(EPhase::Execute);
 	Payload->CommandListsToExecute.Add(CommandList);
 
-	IF_RENDER_TRACING_ENABLED(RenderTracing::D3D12OpenCommandList(Payload, CommandList->BaseCommandList().GetNoRefCount()));
+	IF_RENDER_TRACING_ENABLED(RenderTracing::D3D12::OpenCommandList(Payload, CommandList->BaseCommandList().GetNoRefCount()));
 
 	check(ActiveQueries == 0);
 }
@@ -389,7 +389,7 @@ void FD3D12ContextCommon::CloseCommandList()
 
 	FD3D12Payload* Payload = GetPayload(EPhase::Execute);
 
-	IF_RENDER_TRACING_ENABLED(RenderTracing::D3D12CloseCommandList(Payload, CommandList->BaseCommandList().GetNoRefCount()));
+	IF_RENDER_TRACING_ENABLED(RenderTracing::D3D12::CloseCommandList(Payload, CommandList->BaseCommandList().GetNoRefCount()));
 
 	CommandList->Close();
 	CommandList = nullptr;
@@ -437,7 +437,7 @@ void FD3D12ContextCommon::Finalize(TArray<FD3D12Payload*>& OutPayloads)
 #if UE_RENDER_TRACING_ENABLED
 		for (const FD3D12SyncPointRef& SyncPoint : BatchedSyncPoints.ToWait)
 		{
-			RenderTracing::D3D12WaitSyncPoint(Payload, SyncPoint);
+			RenderTracing::D3D12::WaitSyncPoint(Payload, SyncPoint);
 		}
 #endif
 
@@ -452,7 +452,7 @@ void FD3D12ContextCommon::Finalize(TArray<FD3D12Payload*>& OutPayloads)
 #if UE_RENDER_TRACING_ENABLED
 		for (const FD3D12SyncPointRef& SyncPoint : BatchedSyncPoints.ToSignal)
 		{
-			RenderTracing::D3D12SignalSyncPoint(Payload, SyncPoint);
+			RenderTracing::D3D12::SignalSyncPoint(Payload, SyncPoint);
 		}
 #endif
 
@@ -616,7 +616,7 @@ FD3D12SyncPoint* FD3D12CopyScope::GetSyncPoint() const
 void FD3D12ContextCommon::NewPayload()
 {
 	FD3D12Payload* Payload = new FD3D12Payload(Device->GetQueue(QueueType));
-	IF_RENDER_TRACING_ENABLED(RenderTracing::D3D12AddPayload(this, Payload));
+	IF_RENDER_TRACING_ENABLED(RenderTracing::D3D12::AddPayload(this, Payload));
 	Payloads.Add(Payload);
 }
 
diff --git a/Engine/Source/Runtime/D3D12RHI/Private/D3D12Device.cpp b/Engine/Source/Runtime/D3D12RHI/Private/D3D12Device.cpp
--- a/Engine/Source/Runtime/D3D12RHI/Private/D3D12Device.cpp
+++ b/Engine/Source/Runtime/D3D12RHI/Private/D3D12Device.cpp
@@ -563,7 +563,7 @@ void FD3D12Device::BlockUntilIdle()
 		Payloads.Add(Payload);
 		SyncPoints.Add(SyncPoint);
 
-		IF_RENDER_TRACING_ENABLED(RenderTracing::D3D12SignalSyncPoint(Payload, SyncPoint));
+		IF_RENDER_TRACING_ENABLED(RenderTracing::D3D12::SignalSyncPoint(Payload, SyncPoint));
 	}
 
 	FD3D12DynamicRHI::GetD3DRHI()->SubmitPayloads(MoveTemp(Payloads));
diff --git a/Engine/Source/Runtime/D3D12RHI/Private/D3D12RHI.cpp b/Engine/Source/Runtime/D3D12RHI/Private/D3D12RHI.cpp
--- a/Engine/Source/Runtime/D3D12RHI/Private/D3D12RHI.cpp
+++ b/Engine/Source/Runtime/D3D12RHI/Private/D3D12RHI.cpp
@@ -429,7 +429,7 @@ void FD3D12DynamicRHI::EnqueueEndOfPipeTask(TUniqueFunction<void()> TaskFunc, TU
 		Payload->SyncPointsToSignal.Emplace(SyncPoint);
 		Prereqs.Add(SyncPoint->GetGraphEvent());
 
-		IF_RENDER_TRACING_ENABLED(RenderTracing::D3D12SignalSyncPoint(Payload, SyncPoint));
+		IF_RENDER_TRACING_ENABLED(RenderTracing::D3D12::SignalSyncPoint(Payload, SyncPoint));
 
 		if (ModifyPayloadCallback)
 			ModifyPayloadCallback(*Payload);
diff --git a/Engine/Source/Runtime/D3D12RHI/Private/D3D12Submission.cpp b/Engine/Source/Runtime/D3D12RHI/Private/D3D12Submission.cpp
--- a/Engine/Source/Runtime/D3D12RHI/Private/D3D12Submission.cpp
+++ b/Engine/Source/Runtime/D3D12RHI/Private/D3D12Submission.cpp
@@ -335,7 +335,7 @@ FD3D12DynamicRHI::FProcessResult FD3D12DynamicRHI::ProcessSubmissionQueue()
 	SCOPE_CYCLE_COUNTER(STAT_D3D12Submit);
 	LLM_SCOPE_BYNAME(TEXT("RHIMisc/ProcessSubmissionQueue"));
 	
-	IF_RENDER_TRACING_ENABLED(RenderTracing::D3D12ProcessSubmissionQueueEnter());
+	IF_RENDER_TRACING_ENABLED(RenderTracing::D3D12::ProcessSubmissionQueueEnter());
 
 	FD3D12Queue::FPayloadArray PayloadsToHandDown;
 
@@ -406,7 +406,7 @@ FD3D12DynamicRHI::FProcessResult FD3D12DynamicRHI::ProcessSubmissionQueue()
 							// Skip processing this queue, and move on to the next. We will retry later when
 							// further work is submitted, which may contain the sync point we need.
 							Result.Status |= EQueueStatus::Pending;
-							IF_RENDER_TRACING_ENABLED(RenderTracing::D3D12SubmitYieldOnUnresolvedSyncPoint(CurrentQueue, Payload, SyncPoint));
+							IF_RENDER_TRACING_ENABLED(RenderTracing::D3D12::SubmitYieldOnUnresolvedSyncPoint(CurrentQueue, Payload, SyncPoint));
 							return;
 						}
 
@@ -429,7 +429,7 @@ FD3D12DynamicRHI::FProcessResult FD3D12DynamicRHI::ProcessSubmissionQueue()
 							// Skip processing this queue, and move on to the next. We will retry later when
 							// further work is submitted, which may contain the fence we need.
 							Result.Status |= EQueueStatus::Pending;
-							IF_RENDER_TRACING_ENABLED(RenderTracing::D3D12SubmitYieldOnUnresolvedManualFence(CurrentQueue, Payload, ManualFence.Fence, ManualFence.Value));
+							IF_RENDER_TRACING_ENABLED(RenderTracing::D3D12::SubmitYieldOnUnresolvedManualFence(CurrentQueue, Payload, ManualFence.Fence, ManualFence.Value));
 							return;
 						}
 
@@ -486,7 +486,7 @@ FD3D12DynamicRHI::FProcessResult FD3D12DynamicRHI::ProcessSubmissionQueue()
 		InterruptThread->Kick();
 	}
 
-	IF_RENDER_TRACING_ENABLED(RenderTracing::D3D12ProcessSubmissionQueueExit(static_cast<uint8>(Result.Status)));
+	IF_RENDER_TRACING_ENABLED(RenderTracing::D3D12::ProcessSubmissionQueueExit(static_cast<uint8>(Result.Status)));
 	return Result;
 }
 
@@ -619,7 +619,7 @@ uint64 FD3D12Queue::FinalizePayload(bool bRequiresSignal, FPayloadArray& Payload
 	for (FD3D12SyncPointRef& SyncPoint : PayloadToSubmit->SyncPointsToSignal)
 	{
 		check(!SyncPoint->ResolvedFence.IsSet());
-		IF_RENDER_TRACING_ENABLED(RenderTracing::D3D12SubmitResolveSyncPoint(*this, PayloadToSubmit, SyncPoint, Fence.NextCompletionValue));
+		IF_RENDER_TRACING_ENABLED(RenderTracing::D3D12::SubmitResolveSyncPoint(*this, PayloadToSubmit, SyncPoint, Fence.NextCompletionValue));
 		SyncPoint->ResolvedFence.Emplace(Fence, NextCompletionValue);
 	}
 
@@ -691,7 +691,7 @@ void FD3D12DynamicRHI::FlushBatchedPayloads(FD3D12Queue::FPayloadArray& Payloads
 				, LocalFence.OwnerQueue->GetProfilerQueue()
 			);
 
-			IF_RENDER_TRACING_ENABLED(RenderTracing::D3D12SubmitWaitQueueFence(Payload, Queue, *LocalFence.OwnerQueue, Value));
+			IF_RENDER_TRACING_ENABLED(RenderTracing::D3D12::SubmitWaitQueueFence(Payload, Queue, *LocalFence.OwnerQueue, Value));
 			VERIFYD3D12RESULT(Queue.D3DCommandQueue->Wait(LocalFence.D3DFence, Value));
 		}
 
@@ -699,7 +699,7 @@ void FD3D12DynamicRHI::FlushBatchedPayloads(FD3D12Queue::FPayloadArray& Payloads
 		for (auto& [LocalFence, Value, bFenceSignalled] : Payload->ManualFencesToWait)
 		{
 			checkf(!bFenceSignalled.IsValid() || *bFenceSignalled, TEXT("The signal operation for a provided manual D3D12 fence has not yet happened."));
-			IF_RENDER_TRACING_ENABLED(RenderTracing::D3D12SubmitWaitManualFence(Payload, Queue, LocalFence, Value));
+			IF_RENDER_TRACING_ENABLED(RenderTracing::D3D12::SubmitWaitManualFence(Payload, Queue, LocalFence, Value));
 			VERIFYD3D12RESULT(Queue.D3DCommandQueue->Wait(LocalFence, Value));
 		}
 	};
@@ -740,7 +740,7 @@ void FD3D12DynamicRHI::FlushBatchedPayloads(FD3D12Queue::FPayloadArray& Payloads
 			CommandLists.Append(MoveTemp(Payload->CommandListsToExecute));
 		}
 
-		IF_RENDER_TRACING_ENABLED(RenderTracing::D3D12SubmitPlatformCommandLists(
+		IF_RENDER_TRACING_ENABLED(RenderTracing::D3D12::SubmitPlatformCommandLists(
 			Queue,
 			MakeConstArrayView<FD3D12Payload*>(&PayloadsToSubmit[FirstPayload], LastPayload - FirstPayload),
 			D3DCommandLists
@@ -798,7 +798,7 @@ void FD3D12DynamicRHI::FlushBatchedPayloads(FD3D12Queue::FPayloadArray& Payloads
 		// Signal any manual fences
 		for (auto& [ManualFence, Value, bFenceSignalled] : Payload->ManualFencesToSignal)
 		{
-			IF_RENDER_TRACING_ENABLED(RenderTracing::D3D12SubmitSignalManualFence(Payload, Queue, ManualFence, Value));
+			IF_RENDER_TRACING_ENABLED(RenderTracing::D3D12::SubmitSignalManualFence(Payload, Queue, ManualFence, Value));
 
 			VERIFYD3D12RESULT(Queue.D3DCommandQueue->Signal(ManualFence, Value));
 			if (bFenceSignalled)
@@ -817,7 +817,7 @@ void FD3D12DynamicRHI::FlushBatchedPayloads(FD3D12Queue::FPayloadArray& Payloads
 				, Payload->CompletionFenceValue
 			);
 
-			IF_RENDER_TRACING_ENABLED(RenderTracing::D3D12SubmitSignalQueueFence(Payload, Queue, Payload->CompletionFenceValue));
+			IF_RENDER_TRACING_ENABLED(RenderTracing::D3D12::SubmitSignalQueueFence(Payload, Queue, Payload->CompletionFenceValue));
 
 			VERIFYD3D12RESULT(Queue.D3DCommandQueue->Signal(Queue.Fence.D3DFence, Payload->CompletionFenceValue));
 			Queue.Fence.LastSignaledValue.store(Payload->CompletionFenceValue, std::memory_order_release);
@@ -1057,7 +1057,7 @@ FD3D12DynamicRHI::FProcessResult FD3D12DynamicRHI::ProcessInterruptQueue()
 	SCOPED_NAMED_EVENT_TEXT("InterruptQueue_Process", FColor::Yellow);
 	LLM_SCOPE_BYNAME(TEXT("RHIMisc/ProcessInterruptQueue"));
 
-	IF_RENDER_TRACING_ENABLED(RenderTracing::D3D12ProcessInterruptQueueEnter());
+	IF_RENDER_TRACING_ENABLED(RenderTracing::D3D12::ProcessInterruptQueueEnter());
 
 	// Timer that clamps each tick to prevent false positive GPU timeouts
 	// when a debugger is attached and the process is broken.
@@ -1170,7 +1170,7 @@ FD3D12DynamicRHI::FProcessResult FD3D12DynamicRHI::ProcessInterruptQueue()
 			CurrentQueue.PendingInterrupt.Dequeue();
 			Result.Status |= EQueueStatus::Processed;
 
-			IF_RENDER_TRACING_ENABLED(RenderTracing::D3D12InterruptQueueFenceSignaled(CurrentQueue, Payload, CompletedFenceValue, LastSignaledFenceValue));
+			IF_RENDER_TRACING_ENABLED(RenderTracing::D3D12::InterruptQueueFenceSignaled(CurrentQueue, Payload, CompletedFenceValue, LastSignaledFenceValue));
 
 			// Resolve query results
 			{
@@ -1272,7 +1272,7 @@ FD3D12DynamicRHI::FProcessResult FD3D12DynamicRHI::ProcessInterruptQueue()
 		TerminateOnGPUCrash(GGPUD3DResult);
 	}
 
-	IF_RENDER_TRACING_ENABLED(RenderTracing::D3D12ProcessInterruptQueueExit(static_cast<uint8>(Result.Status)));
+	IF_RENDER_TRACING_ENABLED(RenderTracing::D3D12::ProcessInterruptQueueExit(static_cast<uint8>(Result.Status)));
 	return Result;
 }
 
diff --git a/Engine/Source/Runtime/D3D12RHI/Private/D3D12Tracing.cpp b/Engine/Source/Runtime/D3D12RHI/Private/D3D12Tracing.cpp
--- a/Engine/Source/Runtime/D3D12RHI/Private/D3D12Tracing.cpp
+++ b/Engine/Source/Runtime/D3D12RHI/Private/D3D12Tracing.cpp
@@ -8,130 +8,130 @@
 
 #if UE_RENDER_TRACING_ENABLED
 
-RENDER_TRACE_EVENT_BEGIN(D3D12AssociateContext, Important | NoSync)
+namespace RenderTracing { namespace D3D12
+{
+
+RENDER_TRACE_EVENT_BEGIN(AssociateContext, Important | NoSync)
 	UE_TRACE_EVENT_FIELD(uint64, PlatformContextID)
 	UE_TRACE_EVENT_FIELD(uint64, RHIContextID)
 UE_TRACE_EVENT_END()
 
-RENDER_TRACE_EVENT_BEGIN(D3D12PayloadCreated)
+RENDER_TRACE_EVENT_BEGIN(PayloadCreated)
 	UE_TRACE_EVENT_FIELD(uint64, PayloadID)
 	UE_TRACE_EVENT_FIELD(uint8, Queue)
 UE_TRACE_EVENT_END()
 
-RENDER_TRACE_EVENT_BEGIN(D3D12AddPayload)
+RENDER_TRACE_EVENT_BEGIN(AddPayload)
 	UE_TRACE_EVENT_FIELD(uint64, PlatformContextID)
 	UE_TRACE_EVENT_FIELD(uint64, PayloadID)
 UE_TRACE_EVENT_END()
 
-enum ED3D12CommandListOpType
+enum ECommandListOpType
 {
-	D3D12CommandListOp_Open = 0,
-	D3D12CommandListOp_Close = 1
+	CommandListOp_Open = 0,
+	CommandListOp_Close = 1
 };
 
-RENDER_TRACE_EVENT_BEGIN(D3D12CommandListOp)
+RENDER_TRACE_EVENT_BEGIN(CommandListOp)
 	UE_TRACE_EVENT_FIELD(uint64, PayloadID)
-	UE_TRACE_EVENT_FIELD(uint64, D3D12CmdListID)
+	UE_TRACE_EVENT_FIELD(uint64, PlatformCmdListID)
 	UE_TRACE_EVENT_FIELD(uint8, Op)
 UE_TRACE_EVENT_END()
 
-enum ED3D12FenceOpType
+enum EFenceOpType
 {
-	D3D12FenceOp_Wait = 0,
-	D3D12FenceOp_Signal = 1
+	FenceOp_Wait = 0,
+	FenceOp_Signal = 1
 };
 
-RENDER_TRACE_EVENT_BEGIN(D3D12SyncPointOp)
+RENDER_TRACE_EVENT_BEGIN(SyncPointOp)
 	UE_TRACE_EVENT_FIELD(uint64, PayloadID)
 	UE_TRACE_EVENT_FIELD(uint64, SyncPointID)
 	UE_TRACE_EVENT_FIELD(uint8, SyncPointType)
 	UE_TRACE_EVENT_FIELD(uint8, Op)
 UE_TRACE_EVENT_END()
 
-RENDER_TRACE_EVENT_BEGIN(D3D12ManualFenceOp)
+RENDER_TRACE_EVENT_BEGIN(ManualFenceOp)
 	UE_TRACE_EVENT_FIELD(uint64, PayloadID)
 	UE_TRACE_EVENT_FIELD(uint64, FenceID)
 	UE_TRACE_EVENT_FIELD(uint64, FenceVal)
 	UE_TRACE_EVENT_FIELD(uint8, Op)
 UE_TRACE_EVENT_END()
 
-RENDER_TRACE_EVENT_BEGIN(D3D12ProcessSubmissionQueueEnter)
+RENDER_TRACE_EVENT_BEGIN(ProcessSubmissionQueueEnter)
 UE_TRACE_EVENT_END()
 
-RENDER_TRACE_EVENT_BEGIN(D3D12SubmitYieldOnUnresolvedSyncPoint)
+RENDER_TRACE_EVENT_BEGIN(SubmitYieldOnUnresolvedSyncPoint)
 	UE_TRACE_EVENT_FIELD(uint8, Queue)
 	UE_TRACE_EVENT_FIELD(uint64, PayloadID)
 	UE_TRACE_EVENT_FIELD(uint64, SyncPointID)
 UE_TRACE_EVENT_END()
 
-RENDER_TRACE_EVENT_BEGIN(D3D12SubmitYieldOnUnresolvedManualFence)
+RENDER_TRACE_EVENT_BEGIN(SubmitYieldOnUnresolvedManualFence)
 	UE_TRACE_EVENT_FIELD(uint8, Queue)
 	UE_TRACE_EVENT_FIELD(uint64, PayloadID)
 	UE_TRACE_EVENT_FIELD(uint64, FenceID)
 	UE_TRACE_EVENT_FIELD(uint64, FenceValue)
 UE_TRACE_EVENT_END()
 
-RENDER_TRACE_EVENT_BEGIN(D3D12SubmitWaitQueueFence)
+RENDER_TRACE_EVENT_BEGIN(SubmitWaitQueueFence)
 	UE_TRACE_EVENT_FIELD(uint64, PayloadID)
 	UE_TRACE_EVENT_FIELD(uint8, ExecutingQueue)
 	UE_TRACE_EVENT_FIELD(uint8, WaitOnQueue)
 	UE_TRACE_EVENT_FIELD(uint64, Value)
 UE_TRACE_EVENT_END()
 
-RENDER_TRACE_EVENT_BEGIN(D3D12SubmitWaitManualFence)
+RENDER_TRACE_EVENT_BEGIN(SubmitWaitManualFence)
 	UE_TRACE_EVENT_FIELD(uint64, PayloadID)
 	UE_TRACE_EVENT_FIELD(uint8, ExecutingQueue)
 	UE_TRACE_EVENT_FIELD(uint64, FenceID)
 	UE_TRACE_EVENT_FIELD(uint64, Value)
 UE_TRACE_EVENT_END()
 
-RENDER_TRACE_EVENT_BEGIN(D3D12SubmitPlatformCommandLists)
+RENDER_TRACE_EVENT_BEGIN(SubmitPlatformCommandLists)
 	UE_TRACE_EVENT_FIELD(uint8, Queue)
 	UE_TRACE_EVENT_FIELD(uint64[], PayloadIDs)
-	UE_TRACE_EVENT_FIELD(uint64[], CmdListIDs)
+	UE_TRACE_EVENT_FIELD(uint64[], PlatformCmdListIDs)
 UE_TRACE_EVENT_END()
 
-RENDER_TRACE_EVENT_BEGIN(D3D12SubmitSignalManualFence)
+RENDER_TRACE_EVENT_BEGIN(SubmitSignalManualFence)
 	UE_TRACE_EVENT_FIELD(uint64, PayloadID)
 	UE_TRACE_EVENT_FIELD(uint8, Queue)
 	UE_TRACE_EVENT_FIELD(uint64, FenceID)
 	UE_TRACE_EVENT_FIELD(uint64, Value)
 UE_TRACE_EVENT_END()
 
-RENDER_TRACE_EVENT_BEGIN(D3D12SubmitSignalQueueFence)
+RENDER_TRACE_EVENT_BEGIN(SubmitSignalQueueFence)
 	UE_TRACE_EVENT_FIELD(uint64, PayloadID)
 	UE_TRACE_EVENT_FIELD(uint8, Queue)
 	UE_TRACE_EVENT_FIELD(uint64, Value)
 UE_TRACE_EVENT_END()
 
-RENDER_TRACE_EVENT_BEGIN(D3D12SubmitResolveSyncPoint)
+RENDER_TRACE_EVENT_BEGIN(SubmitResolveSyncPoint)
 	UE_TRACE_EVENT_FIELD(uint8, Queue)
 	UE_TRACE_EVENT_FIELD(uint64, PayloadID)
 	UE_TRACE_EVENT_FIELD(uint64, SyncPointID)
 	UE_TRACE_EVENT_FIELD(uint64, Value)
 UE_TRACE_EVENT_END()
 
-RENDER_TRACE_EVENT_BEGIN(D3D12ProcessSubmissionQueueExit)
+RENDER_TRACE_EVENT_BEGIN(ProcessSubmissionQueueExit)
 	UE_TRACE_EVENT_FIELD(uint8, Status)
 UE_TRACE_EVENT_END()
 
-RENDER_TRACE_EVENT_BEGIN(D3D12ProcessInterruptQueueEnter)
+RENDER_TRACE_EVENT_BEGIN(ProcessInterruptQueueEnter)
 UE_TRACE_EVENT_END()
 
-RENDER_TRACE_EVENT_BEGIN(D3D12InterruptQueueFenceSignaled)
+RENDER_TRACE_EVENT_BEGIN(InterruptQueueFenceSignaled)
 	UE_TRACE_EVENT_FIELD(uint8, Queue)
 	UE_TRACE_EVENT_FIELD(uint64, PendingPayloadID)
 	UE_TRACE_EVENT_FIELD(uint64, CurrentFenceValue)
 	UE_TRACE_EVENT_FIELD(uint64, LastCPUSignaledFenceValue)
 UE_TRACE_EVENT_END()
 
-RENDER_TRACE_EVENT_BEGIN(D3D12ProcessInterruptQueueExit)
+RENDER_TRACE_EVENT_BEGIN(ProcessInterruptQueueExit)
 	UE_TRACE_EVENT_FIELD(uint8, Status)
 UE_TRACE_EVENT_END()
 
-namespace RenderTracing
-{
-
 static uint8 GetRenderTraceQueueIndex(const FD3D12Queue& Queue)
 {
 	// No mGPU support for now, and only one queue per type.
@@ -144,277 +144,277 @@ static uint8 GetRenderTraceQueueIndex(const FD3D12Queue& Queue)
 	}
 }
 
-void D3D12AssociateContext(const FD3D12ContextCommon* D3D12Context, const IRHIComputeContext* RHIContext)
+void AssociateContext(const FD3D12ContextCommon* D3D12Context, const IRHIComputeContext* RHIContext)
 {
 	if (!IsEnabled(ERenderTracingChannels::RHISubmission))
 	{
 		return;
 	}
 
-	RENDER_TRACE_LOG(D3D12AssociateContext)
-		<< RENDER_TRACE_VALUE(D3D12AssociateContext, PlatformContextID, reinterpret_cast<uintptr_t>(D3D12Context))
-		<< RENDER_TRACE_VALUE(D3D12AssociateContext, RHIContextID, reinterpret_cast<uintptr_t>(RHIContext));
+	RENDER_TRACE_LOG(AssociateContext)
+		<< RENDER_TRACE_VALUE(AssociateContext, PlatformContextID, reinterpret_cast<uintptr_t>(D3D12Context))
+		<< RENDER_TRACE_VALUE(AssociateContext, RHIContextID, reinterpret_cast<uintptr_t>(RHIContext));
 }
 
-void D3D12PayloadCreated(const FD3D12Payload* Payload, const FD3D12Queue& Queue)
+void PayloadCreated(const FD3D12Payload* Payload, const FD3D12Queue& Queue)
 {
 	if (!IsEnabled(ERenderTracingChannels::RHISubmission))
 	{
 		return;
 	}
 
-	RENDER_TRACE_LOG(D3D12PayloadCreated)
-		<< RENDER_TRACE_VALUE(D3D12PayloadCreated, PayloadID, reinterpret_cast<uintptr_t>(Payload))
-		<< RENDER_TRACE_VALUE(D3D12PayloadCreated, Queue, GetRenderTraceQueueIndex(Queue));
+	RENDER_TRACE_LOG(PayloadCreated)
+		<< RENDER_TRACE_VALUE(PayloadCreated, PayloadID, reinterpret_cast<uintptr_t>(Payload))
+		<< RENDER_TRACE_VALUE(PayloadCreated, Queue, GetRenderTraceQueueIndex(Queue));
 }
 
-void D3D12AddPayload(const FD3D12ContextCommon* D3D12Context, const FD3D12Payload* Payload)
+void AddPayload(const FD3D12ContextCommon* D3D12Context, const FD3D12Payload* Payload)
 {
 	if (!IsEnabled(ERenderTracingChannels::RHISubmission))
 	{
 		return;
 	}
 
-	RENDER_TRACE_LOG(D3D12AddPayload)
-		<< RENDER_TRACE_VALUE(D3D12AddPayload, PlatformContextID, reinterpret_cast<uintptr_t>(D3D12Context))
-		<< RENDER_TRACE_VALUE(D3D12AddPayload, PayloadID, reinterpret_cast<uintptr_t>(Payload));
+	RENDER_TRACE_LOG(AddPayload)
+		<< RENDER_TRACE_VALUE(AddPayload, PlatformContextID, reinterpret_cast<uintptr_t>(D3D12Context))
+		<< RENDER_TRACE_VALUE(AddPayload, PayloadID, reinterpret_cast<uintptr_t>(Payload));
 }
 
-static void D3D12CommandListOp(const FD3D12Payload* Payload, const ID3D12CommandList* D3D12CmdList, ED3D12CommandListOpType OpType)
+static void CommandListOp(const FD3D12Payload* Payload, const ID3D12CommandList* D3D12CmdList, ECommandListOpType OpType)
 {
 	if (!IsEnabled(ERenderTracingChannels::RHISubmission))
 	{
 		return;
 	}
 
-	RENDER_TRACE_LOG(D3D12CommandListOp)
-		<< RENDER_TRACE_VALUE(D3D12CommandListOp, PayloadID, reinterpret_cast<uintptr_t>(Payload))
-		<< RENDER_TRACE_VALUE(D3D12CommandListOp, D3D12CmdListID, reinterpret_cast<uintptr_t>(D3D12CmdList))
-		<< RENDER_TRACE_VALUE(D3D12CommandListOp, Op, OpType);
+	RENDER_TRACE_LOG(CommandListOp)
+		<< RENDER_TRACE_VALUE(CommandListOp, PayloadID, reinterpret_cast<uintptr_t>(Payload))
+		<< RENDER_TRACE_VALUE(CommandListOp, PlatformCmdListID, reinterpret_cast<uintptr_t>(D3D12CmdList))
+		<< RENDER_TRACE_VALUE(CommandListOp, Op, OpType);
 }
 
-void D3D12OpenCommandList(const FD3D12Payload* Payload, const ID3D12CommandList* D3D12CmdList)
+void OpenCommandList(const FD3D12Payload* Payload, const ID3D12CommandList* D3D12CmdList)
 {
-	D3D12CommandListOp(Payload, D3D12CmdList, D3D12CommandListOp_Open);
+	CommandListOp(Payload, D3D12CmdList, CommandListOp_Open);
 }
 
-void D3D12CloseCommandList(const FD3D12Payload* Payload, const ID3D12CommandList* D3D12CmdList)
+void CloseCommandList(const FD3D12Payload* Payload, const ID3D12CommandList* D3D12CmdList)
 {
-	D3D12CommandListOp(Payload, D3D12CmdList, D3D12CommandListOp_Close);
+	CommandListOp(Payload, D3D12CmdList, CommandListOp_Close);
 }
 
-static void D3D12SyncPointOp(const FD3D12Payload* Payload, const FD3D12SyncPoint* SyncPoint, ED3D12FenceOpType OpType)
+static void SyncPointOp(const FD3D12Payload* Payload, const FD3D12SyncPoint* SyncPoint, EFenceOpType OpType)
 {
 	if (!IsEnabled(ERenderTracingChannels::RHISubmission))
 	{
 		return;
 	}
 
-	RENDER_TRACE_LOG(D3D12SyncPointOp)
-		<< RENDER_TRACE_VALUE(D3D12SyncPointOp, PayloadID, reinterpret_cast<uintptr_t>(Payload))
-		<< RENDER_TRACE_VALUE(D3D12SyncPointOp, SyncPointID, reinterpret_cast<uintptr_t>(SyncPoint))
-		<< RENDER_TRACE_VALUE(D3D12SyncPointOp, SyncPointType, static_cast<uint8>(SyncPoint->GetType()))
-		<< RENDER_TRACE_VALUE(D3D12SyncPointOp, Op, OpType);
+	RENDER_TRACE_LOG(SyncPointOp)
+		<< RENDER_TRACE_VALUE(SyncPointOp, PayloadID, reinterpret_cast<uintptr_t>(Payload))
+		<< RENDER_TRACE_VALUE(SyncPointOp, SyncPointID, reinterpret_cast<uintptr_t>(SyncPoint))
+		<< RENDER_TRACE_VALUE(SyncPointOp, SyncPointType, static_cast<uint8>(SyncPoint->GetType()))
+		<< RENDER_TRACE_VALUE(SyncPointOp, Op, OpType);
 }
 
-void D3D12WaitSyncPoint(const FD3D12Payload* Payload, const FD3D12SyncPoint* SyncPoint)
+void WaitSyncPoint(const FD3D12Payload* Payload, const FD3D12SyncPoint* SyncPoint)
 {
-	D3D12SyncPointOp(Payload, SyncPoint, D3D12FenceOp_Wait);
+	SyncPointOp(Payload, SyncPoint, FenceOp_Wait);
 }
 
-void D3D12SignalSyncPoint(const FD3D12Payload* Payload, const FD3D12SyncPoint* SyncPoint)
+void SignalSyncPoint(const FD3D12Payload* Payload, const FD3D12SyncPoint* SyncPoint)
 {
-	D3D12SyncPointOp(Payload, SyncPoint, D3D12FenceOp_Signal);
+	SyncPointOp(Payload, SyncPoint, FenceOp_Signal);
 }
 
-static void D3D12ManualFenceOp(const FD3D12Payload* Payload, const ID3D12Fence* Fence, uint64 Value, ED3D12FenceOpType OpType)
+static void ManualFenceOp(const FD3D12Payload* Payload, const ID3D12Fence* Fence, uint64 Value, EFenceOpType OpType)
 {
 	if (!IsEnabled(ERenderTracingChannels::RHISubmission))
 	{
 		return;
 	}
 
-	RENDER_TRACE_LOG(D3D12ManualFenceOp)
-		<< RENDER_TRACE_VALUE(D3D12ManualFenceOp, PayloadID, reinterpret_cast<uintptr_t>(Payload))
-		<< RENDER_TRACE_VALUE(D3D12ManualFenceOp, FenceID, reinterpret_cast<uintptr_t>(Fence))
-		<< RENDER_TRACE_VALUE(D3D12ManualFenceOp, FenceVal, Value)
-		<< RENDER_TRACE_VALUE(D3D12ManualFenceOp, Op, OpType);
+	RENDER_TRACE_LOG(ManualFenceOp)
+		<< RENDER_TRACE_VALUE(ManualFenceOp, PayloadID, reinterpret_cast<uintptr_t>(Payload))
+		<< RENDER_TRACE_VALUE(ManualFenceOp, FenceID, reinterpret_cast<uintptr_t>(Fence))
+		<< RENDER_TRACE_VALUE(ManualFenceOp, FenceVal, Value)
+		<< RENDER_TRACE_VALUE(ManualFenceOp, Op, OpType);
 }
 
-void D3D12WaitManualFence(const FD3D12Payload* Payload, const ID3D12Fence* Fence, uint64 Value)
+void WaitManualFence(const FD3D12Payload* Payload, const ID3D12Fence* Fence, uint64 Value)
 {
-	D3D12ManualFenceOp(Payload, Fence, Value, D3D12FenceOp_Wait);
+	ManualFenceOp(Payload, Fence, Value, FenceOp_Wait);
 }
 
-void D3D12SignalManualFence(const FD3D12Payload* Payload, const ID3D12Fence* Fence, uint64 Value)
+void SignalManualFence(const FD3D12Payload* Payload, const ID3D12Fence* Fence, uint64 Value)
 {
-	D3D12ManualFenceOp(Payload, Fence, Value, D3D12FenceOp_Signal);
+	ManualFenceOp(Payload, Fence, Value, FenceOp_Signal);
 }
 
-void D3D12ProcessSubmissionQueueEnter()
+void ProcessSubmissionQueueEnter()
 {
 	if (!IsEnabled(ERenderTracingChannels::RHISubmission))
 	{
 		return;
 	}
 
-	RENDER_TRACE_LOG(D3D12ProcessSubmissionQueueEnter);
+	RENDER_TRACE_LOG(ProcessSubmissionQueueEnter);
 }
 
-void D3D12SubmitYieldOnUnresolvedSyncPoint(const FD3D12Queue& Queue, const FD3D12Payload* Payload, const FD3D12SyncPoint* SyncPoint)
+void SubmitYieldOnUnresolvedSyncPoint(const FD3D12Queue& Queue, const FD3D12Payload* Payload, const FD3D12SyncPoint* SyncPoint)
 {
 	if (!IsEnabled(ERenderTracingChannels::RHISubmission))
 	{
 		return;
 	}
 
-	RENDER_TRACE_LOG(D3D12SubmitYieldOnUnresolvedSyncPoint)
-		<< RENDER_TRACE_VALUE(D3D12SubmitYieldOnUnresolvedSyncPoint, Queue, GetRenderTraceQueueIndex(Queue))
-		<< RENDER_TRACE_VALUE(D3D12SubmitYieldOnUnresolvedSyncPoint, PayloadID, reinterpret_cast<uintptr_t>(Payload))
-		<< RENDER_TRACE_VALUE(D3D12SubmitYieldOnUnresolvedSyncPoint, SyncPointID, reinterpret_cast<uintptr_t>(SyncPoint));
+	RENDER_TRACE_LOG(SubmitYieldOnUnresolvedSyncPoint)
+		<< RENDER_TRACE_VALUE(SubmitYieldOnUnresolvedSyncPoint, Queue, GetRenderTraceQueueIndex(Queue))
+		<< RENDER_TRACE_VALUE(SubmitYieldOnUnresolvedSyncPoint, PayloadID, reinterpret_cast<uintptr_t>(Payload))
+		<< RENDER_TRACE_VALUE(SubmitYieldOnUnresolvedSyncPoint, SyncPointID, reinterpret_cast<uintptr_t>(SyncPoint));
 }
 
-void D3D12SubmitYieldOnUnresolvedManualFence(const FD3D12Queue& Queue, const FD3D12Payload* Payload, const ID3D12Fence* Fence, uint64 Value)
+void SubmitYieldOnUnresolvedManualFence(const FD3D12Queue& Queue, const FD3D12Payload* Payload, const ID3D12Fence* Fence, uint64 Value)
 {
 	if (!IsEnabled(ERenderTracingChannels::RHISubmission))
 	{
 		return;
 	}
 
-	RENDER_TRACE_LOG(D3D12SubmitYieldOnUnresolvedManualFence)
-		<< RENDER_TRACE_VALUE(D3D12SubmitYieldOnUnresolvedManualFence, Queue, GetRenderTraceQueueIndex(Queue))
-		<< RENDER_TRACE_VALUE(D3D12SubmitYieldOnUnresolvedManualFence, PayloadID, reinterpret_cast<uintptr_t>(Payload))
-		<< RENDER_TRACE_VALUE(D3D12SubmitYieldOnUnresolvedManualFence, FenceID, reinterpret_cast<uintptr_t>(Fence))
-		<< RENDER_TRACE_VALUE(D3D12SubmitYieldOnUnresolvedManualFence, FenceValue, Value);
+	RENDER_TRACE_LOG(SubmitYieldOnUnresolvedManualFence)
+		<< RENDER_TRACE_VALUE(SubmitYieldOnUnresolvedManualFence, Queue, GetRenderTraceQueueIndex(Queue))
+		<< RENDER_TRACE_VALUE(SubmitYieldOnUnresolvedManualFence, PayloadID, reinterpret_cast<uintptr_t>(Payload))
+		<< RENDER_TRACE_VALUE(SubmitYieldOnUnresolvedManualFence, FenceID, reinterpret_cast<uintptr_t>(Fence))
+		<< RENDER_TRACE_VALUE(SubmitYieldOnUnresolvedManualFence, FenceValue, Value);
 }
 
-void D3D12SubmitWaitQueueFence(const FD3D12Payload* Payload, const FD3D12Queue& ExecutingQueue, const FD3D12Queue& WaitOnQueue, uint64 Value)
+void SubmitWaitQueueFence(const FD3D12Payload* Payload, const FD3D12Queue& ExecutingQueue, const FD3D12Queue& WaitOnQueue, uint64 Value)
 {
 	if (!IsEnabled(ERenderTracingChannels::RHISubmission))
 	{
 		return;
 	}
 
-	RENDER_TRACE_LOG(D3D12SubmitWaitQueueFence)
-		<< RENDER_TRACE_VALUE(D3D12SubmitWaitQueueFence, PayloadID, reinterpret_cast<uintptr_t>(Payload))
-		<< RENDER_TRACE_VALUE(D3D12SubmitWaitQueueFence, ExecutingQueue, GetRenderTraceQueueIndex(ExecutingQueue))
-		<< RENDER_TRACE_VALUE(D3D12SubmitWaitQueueFence, WaitOnQueue, GetRenderTraceQueueIndex(WaitOnQueue))
-		<< RENDER_TRACE_VALUE(D3D12SubmitWaitQueueFence, Value, Value);
+	RENDER_TRACE_LOG(SubmitWaitQueueFence)
+		<< RENDER_TRACE_VALUE(SubmitWaitQueueFence, PayloadID, reinterpret_cast<uintptr_t>(Payload))
+		<< RENDER_TRACE_VALUE(SubmitWaitQueueFence, ExecutingQueue, GetRenderTraceQueueIndex(ExecutingQueue))
+		<< RENDER_TRACE_VALUE(SubmitWaitQueueFence, WaitOnQueue, GetRenderTraceQueueIndex(WaitOnQueue))
+		<< RENDER_TRACE_VALUE(SubmitWaitQueueFence, Value, Value);
 }
 
-void D3D12SubmitWaitManualFence(const FD3D12Payload* Payload, const FD3D12Queue& ExecutingQueue, const ID3D12Fence* Fence, uint64 Value)
+void SubmitWaitManualFence(const FD3D12Payload* Payload, const FD3D12Queue& ExecutingQueue, const ID3D12Fence* Fence, uint64 Value)
 {
 	if (!IsEnabled(ERenderTracingChannels::RHISubmission))
 	{
 		return;
 	}
 
-	RENDER_TRACE_LOG(D3D12SubmitWaitManualFence)
-		<< RENDER_TRACE_VALUE(D3D12SubmitWaitManualFence, PayloadID, reinterpret_cast<uintptr_t>(Payload))
-		<< RENDER_TRACE_VALUE(D3D12SubmitWaitManualFence, ExecutingQueue, GetRenderTraceQueueIndex(ExecutingQueue))
-		<< RENDER_TRACE_VALUE(D3D12SubmitWaitManualFence, FenceID, reinterpret_cast<uintptr_t>(Fence))
-		<< RENDER_TRACE_VALUE(D3D12SubmitWaitManualFence, Value, Value);
+	RENDER_TRACE_LOG(SubmitWaitManualFence)
+		<< RENDER_TRACE_VALUE(SubmitWaitManualFence, PayloadID, reinterpret_cast<uintptr_t>(Payload))
+		<< RENDER_TRACE_VALUE(SubmitWaitManualFence, ExecutingQueue, GetRenderTraceQueueIndex(ExecutingQueue))
+		<< RENDER_TRACE_VALUE(SubmitWaitManualFence, FenceID, reinterpret_cast<uintptr_t>(Fence))
+		<< RENDER_TRACE_VALUE(SubmitWaitManualFence, Value, Value);
 }
 
-void D3D12SubmitPlatformCommandLists(const FD3D12Queue& Queue, TConstArrayView<FD3D12Payload*> Payloads, TConstArrayView<ID3D12CommandList*> CmdLists)
+void SubmitPlatformCommandLists(const FD3D12Queue& Queue, TConstArrayView<FD3D12Payload*> Payloads, TConstArrayView<ID3D12CommandList*> CmdLists)
 {
 	if (!IsEnabled(ERenderTracingChannels::RHISubmission))
 	{
 		return;
 	}
 
-	RENDER_TRACE_LOG(D3D12SubmitPlatformCommandLists)
-		<< RENDER_TRACE_VALUE(D3D12SubmitPlatformCommandLists, Queue, GetRenderTraceQueueIndex(Queue))
-		<< RENDER_TRACE_ARRAY_VALUE(D3D12SubmitPlatformCommandLists, PayloadIDs, Payloads, uint64)
-		<< RENDER_TRACE_ARRAY_VALUE(D3D12SubmitPlatformCommandLists, CmdListIDs, CmdLists, uint64);
+	RENDER_TRACE_LOG(SubmitPlatformCommandLists)
+		<< RENDER_TRACE_VALUE(SubmitPlatformCommandLists, Queue, GetRenderTraceQueueIndex(Queue))
+		<< RENDER_TRACE_ARRAY_VALUE(SubmitPlatformCommandLists, PayloadIDs, Payloads, uint64)
+		<< RENDER_TRACE_ARRAY_VALUE(SubmitPlatformCommandLists, PlatformCmdListIDs, CmdLists, uint64);
 }
 
-void D3D12SubmitSignalManualFence(const FD3D12Payload* Payload, const FD3D12Queue& Queue, const ID3D12Fence* Fence, uint64 Value)
+void SubmitSignalManualFence(const FD3D12Payload* Payload, const FD3D12Queue& Queue, const ID3D12Fence* Fence, uint64 Value)
 {
 	if (!IsEnabled(ERenderTracingChannels::RHISubmission))
 	{
 		return;
 	}
 
-	RENDER_TRACE_LOG(D3D12SubmitSignalManualFence)
-		<< RENDER_TRACE_VALUE(D3D12SubmitSignalManualFence, PayloadID, reinterpret_cast<uintptr_t>(Payload))
-		<< RENDER_TRACE_VALUE(D3D12SubmitSignalManualFence, Queue, GetRenderTraceQueueIndex(Queue))
-		<< RENDER_TRACE_VALUE(D3D12SubmitSignalManualFence, FenceID, reinterpret_cast<uintptr_t>(Fence))
-		<< RENDER_TRACE_VALUE(D3D12SubmitSignalManualFence, Value, Value);
+	RENDER_TRACE_LOG(SubmitSignalManualFence)
+		<< RENDER_TRACE_VALUE(SubmitSignalManualFence, PayloadID, reinterpret_cast<uintptr_t>(Payload))
+		<< RENDER_TRACE_VALUE(SubmitSignalManualFence, Queue, GetRenderTraceQueueIndex(Queue))
+		<< RENDER_TRACE_VALUE(SubmitSignalManualFence, FenceID, reinterpret_cast<uintptr_t>(Fence))
+		<< RENDER_TRACE_VALUE(SubmitSignalManualFence, Value, Value);
 }
 
-void D3D12SubmitSignalQueueFence(const FD3D12Payload* Payload, const FD3D12Queue& Queue, uint64 Value)
+void SubmitSignalQueueFence(const FD3D12Payload* Payload, const FD3D12Queue& Queue, uint64 Value)
 {
 	if (!IsEnabled(ERenderTracingChannels::RHISubmission))
 	{
 		return;
 	}
 
-	RENDER_TRACE_LOG(D3D12SubmitSignalQueueFence)
-		<< RENDER_TRACE_VALUE(D3D12SubmitSignalQueueFence, PayloadID, reinterpret_cast<uintptr_t>(Payload))
-		<< RENDER_TRACE_VALUE(D3D12SubmitSignalQueueFence, Queue, GetRenderTraceQueueIndex(Queue))
-		<< RENDER_TRACE_VALUE(D3D12SubmitSignalQueueFence, Value, Value);
+	RENDER_TRACE_LOG(SubmitSignalQueueFence)
+		<< RENDER_TRACE_VALUE(SubmitSignalQueueFence, PayloadID, reinterpret_cast<uintptr_t>(Payload))
+		<< RENDER_TRACE_VALUE(SubmitSignalQueueFence, Queue, GetRenderTraceQueueIndex(Queue))
+		<< RENDER_TRACE_VALUE(SubmitSignalQueueFence, Value, Value);
 }
 
-void D3D12SubmitResolveSyncPoint(const FD3D12Queue& Queue, const FD3D12Payload* Payload, const FD3D12SyncPoint* SyncPoint, uint64 Value)
+void SubmitResolveSyncPoint(const FD3D12Queue& Queue, const FD3D12Payload* Payload, const FD3D12SyncPoint* SyncPoint, uint64 Value)
 {
 	if (!IsEnabled(ERenderTracingChannels::RHISubmission))
 	{
 		return;
 	}
 
-	RENDER_TRACE_LOG(D3D12SubmitResolveSyncPoint)
-		<< RENDER_TRACE_VALUE(D3D12SubmitResolveSyncPoint, Queue, GetRenderTraceQueueIndex(Queue))
-		<< RENDER_TRACE_VALUE(D3D12SubmitResolveSyncPoint, PayloadID, reinterpret_cast<uintptr_t>(Payload))
-		<< RENDER_TRACE_VALUE(D3D12SubmitResolveSyncPoint, SyncPointID, reinterpret_cast<uintptr_t>(SyncPoint))
-		<< RENDER_TRACE_VALUE(D3D12SubmitResolveSyncPoint, Value, Value);
+	RENDER_TRACE_LOG(SubmitResolveSyncPoint)
+		<< RENDER_TRACE_VALUE(SubmitResolveSyncPoint, Queue, GetRenderTraceQueueIndex(Queue))
+		<< RENDER_TRACE_VALUE(SubmitResolveSyncPoint, PayloadID, reinterpret_cast<uintptr_t>(Payload))
+		<< RENDER_TRACE_VALUE(SubmitResolveSyncPoint, SyncPointID, reinterpret_cast<uintptr_t>(SyncPoint))
+		<< RENDER_TRACE_VALUE(SubmitResolveSyncPoint, Value, Value);
 }
 
-void D3D12ProcessSubmissionQueueExit(uint8 Status)
+void ProcessSubmissionQueueExit(uint8 Status)
 {
 	if (!IsEnabled(ERenderTracingChannels::RHISubmission))
 	{
 		return;
 	}
 
-	RENDER_TRACE_LOG(D3D12ProcessSubmissionQueueExit)
-		<< RENDER_TRACE_VALUE(D3D12ProcessSubmissionQueueExit, Status, Status);
+	RENDER_TRACE_LOG(ProcessSubmissionQueueExit)
+		<< RENDER_TRACE_VALUE(ProcessSubmissionQueueExit, Status, Status);
 }
 
-void D3D12ProcessInterruptQueueEnter()
+void ProcessInterruptQueueEnter()
 {
 	if (!IsEnabled(ERenderTracingChannels::RHISubmission))
 	{
 		return;
 	}
 
-	RENDER_TRACE_LOG(D3D12ProcessInterruptQueueEnter);
+	RENDER_TRACE_LOG(ProcessInterruptQueueEnter);
 }
 
-void D3D12InterruptQueueFenceSignaled(const FD3D12Queue& Queue, FD3D12Payload* PendingPayload, uint64 CurrentFenceValue, uint64 LastCPUSignaledFenceValue)
+void InterruptQueueFenceSignaled(const FD3D12Queue& Queue, FD3D12Payload* PendingPayload, uint64 CurrentFenceValue, uint64 LastCPUSignaledFenceValue)
 {
 	if (!IsEnabled(ERenderTracingChannels::RHISubmission))
 	{
 		return;
 	}
 
-	RENDER_TRACE_LOG(D3D12InterruptQueueFenceSignaled)
-		<< RENDER_TRACE_VALUE(D3D12InterruptQueueFenceSignaled, Queue, GetRenderTraceQueueIndex(Queue))
-		<< RENDER_TRACE_VALUE(D3D12InterruptQueueFenceSignaled, PendingPayloadID, reinterpret_cast<uintptr_t>(PendingPayload))
-		<< RENDER_TRACE_VALUE(D3D12InterruptQueueFenceSignaled, CurrentFenceValue, CurrentFenceValue)
-		<< RENDER_TRACE_VALUE(D3D12InterruptQueueFenceSignaled, LastCPUSignaledFenceValue, LastCPUSignaledFenceValue);
+	RENDER_TRACE_LOG(InterruptQueueFenceSignaled)
+		<< RENDER_TRACE_VALUE(InterruptQueueFenceSignaled, Queue, GetRenderTraceQueueIndex(Queue))
+		<< RENDER_TRACE_VALUE(InterruptQueueFenceSignaled, PendingPayloadID, reinterpret_cast<uintptr_t>(PendingPayload))
+		<< RENDER_TRACE_VALUE(InterruptQueueFenceSignaled, CurrentFenceValue, CurrentFenceValue)
+		<< RENDER_TRACE_VALUE(InterruptQueueFenceSignaled, LastCPUSignaledFenceValue, LastCPUSignaledFenceValue);
 }
 
-void D3D12ProcessInterruptQueueExit(uint8 Status)
+void ProcessInterruptQueueExit(uint8 Status)
 {
 	if (!IsEnabled(ERenderTracingChannels::RHISubmission))
 	{
 		return;
 	}
 
-	RENDER_TRACE_LOG(D3D12ProcessInterruptQueueExit)
-		<< RENDER_TRACE_VALUE(D3D12ProcessInterruptQueueExit, Status, Status);
+	RENDER_TRACE_LOG(ProcessInterruptQueueExit)
+		<< RENDER_TRACE_VALUE(ProcessInterruptQueueExit, Status, Status);
 }
 
-}
+} }
 #endif
diff --git a/Engine/Source/Runtime/D3D12RHI/Private/D3D12Tracing.h b/Engine/Source/Runtime/D3D12RHI/Private/D3D12Tracing.h
--- a/Engine/Source/Runtime/D3D12RHI/Private/D3D12Tracing.h
+++ b/Engine/Source/Runtime/D3D12RHI/Private/D3D12Tracing.h
@@ -14,34 +14,34 @@ class FD3D12SyncPoint;
 struct ID3D12CommandList;
 struct ID3D12Fence;
 
-namespace RenderTracing
+namespace RenderTracing { namespace D3D12
 {
 
-void D3D12AssociateContext(const FD3D12ContextCommon* D3D12Context, const IRHIComputeContext* RHIContext);
-void D3D12PayloadCreated(const FD3D12Payload* Payload, const FD3D12Queue& Queue);
-void D3D12AddPayload(const FD3D12ContextCommon* D3D12Context, const FD3D12Payload* Payload);
-void D3D12OpenCommandList(const FD3D12Payload* Payload, const ID3D12CommandList* D3D12CmdList);
-void D3D12CloseCommandList(const FD3D12Payload* Payload, const ID3D12CommandList* D3D12CmdList);
-void D3D12WaitSyncPoint(const FD3D12Payload* Payload, const FD3D12SyncPoint* SyncPoint);
-void D3D12SignalSyncPoint(const FD3D12Payload* Payload, const FD3D12SyncPoint* SyncPoint);
-void D3D12WaitManualFence(const FD3D12Payload* Payload, const ID3D12Fence* Fence, uint64 Value);
-void D3D12SignalManualFence(const FD3D12Payload* Payload, const ID3D12Fence* Fence, uint64 Value);
-
-void D3D12ProcessSubmissionQueueEnter();
-void D3D12SubmitYieldOnUnresolvedSyncPoint(const FD3D12Queue& Queue, const FD3D12Payload* Payload, const FD3D12SyncPoint* SyncPoint);
-void D3D12SubmitYieldOnUnresolvedManualFence(const FD3D12Queue& Queue, const FD3D12Payload* Payload, const ID3D12Fence* Fence, uint64 Value);
-void D3D12SubmitWaitQueueFence(const FD3D12Payload* Payload, const FD3D12Queue& ExecutingQueue, const FD3D12Queue& WaitOnQueue, uint64 Value);
-void D3D12SubmitWaitManualFence(const FD3D12Payload* Payload, const FD3D12Queue& ExecutingQueue, const ID3D12Fence* Fence, uint64 Value);
-void D3D12SubmitPlatformCommandLists(const FD3D12Queue& Queue, TConstArrayView<FD3D12Payload*> Payloads, TConstArrayView<ID3D12CommandList*> CmdLists);
-void D3D12SubmitSignalManualFence(const FD3D12Payload* Payload, const FD3D12Queue& Queue, const ID3D12Fence* Fence, uint64 Value);
-void D3D12SubmitSignalQueueFence(const FD3D12Payload* Payload, const FD3D12Queue& Queue, uint64 Value);
-void D3D12SubmitResolveSyncPoint(const FD3D12Queue& Queue, const FD3D12Payload* Payload, const FD3D12SyncPoint* SyncPoint, uint64 Value);
-void D3D12ProcessSubmissionQueueExit(uint8 Status);
-
-void D3D12ProcessInterruptQueueEnter();
-void D3D12InterruptQueueFenceSignaled(const FD3D12Queue& Queue, FD3D12Payload* PendingPayload, uint64 CurrentFenceValue, uint64 LastCPUSignaledFenceValue);
-void D3D12ProcessInterruptQueueExit(uint8 Status);
-
-}
+void AssociateContext(const FD3D12ContextCommon* D3D12Context, const IRHIComputeContext* RHIContext);
+void PayloadCreated(const FD3D12Payload* Payload, const FD3D12Queue& Queue);
+void AddPayload(const FD3D12ContextCommon* D3D12Context, const FD3D12Payload* Payload);
+void OpenCommandList(const FD3D12Payload* Payload, const ID3D12CommandList* D3D12CmdList);
+void CloseCommandList(const FD3D12Payload* Payload, const ID3D12CommandList* D3D12CmdList);
+void WaitSyncPoint(const FD3D12Payload* Payload, const FD3D12SyncPoint* SyncPoint);
+void SignalSyncPoint(const FD3D12Payload* Payload, const FD3D12SyncPoint* SyncPoint);
+void WaitManualFence(const FD3D12Payload* Payload, const ID3D12Fence* Fence, uint64 Value);
+void SignalManualFence(const FD3D12Payload* Payload, const ID3D12Fence* Fence, uint64 Value);
+
+void ProcessSubmissionQueueEnter();
+void SubmitYieldOnUnresolvedSyncPoint(const FD3D12Queue& Queue, const FD3D12Payload* Payload, const FD3D12SyncPoint* SyncPoint);
+void SubmitYieldOnUnresolvedManualFence(const FD3D12Queue& Queue, const FD3D12Payload* Payload, const ID3D12Fence* Fence, uint64 Value);
+void SubmitWaitQueueFence(const FD3D12Payload* Payload, const FD3D12Queue& ExecutingQueue, const FD3D12Queue& WaitOnQueue, uint64 Value);
+void SubmitWaitManualFence(const FD3D12Payload* Payload, const FD3D12Queue& ExecutingQueue, const ID3D12Fence* Fence, uint64 Value);
+void SubmitPlatformCommandLists(const FD3D12Queue& Queue, TConstArrayView<FD3D12Payload*> Payloads, TConstArrayView<ID3D12CommandList*> CmdLists);
+void SubmitSignalManualFence(const FD3D12Payload* Payload, const FD3D12Queue& Queue, const ID3D12Fence* Fence, uint64 Value);
+void SubmitSignalQueueFence(const FD3D12Payload* Payload, const FD3D12Queue& Queue, uint64 Value);
+void SubmitResolveSyncPoint(const FD3D12Queue& Queue, const FD3D12Payload* Payload, const FD3D12SyncPoint* SyncPoint, uint64 Value);
+void ProcessSubmissionQueueExit(uint8 Status);
+
+void ProcessInterruptQueueEnter();
+void InterruptQueueFenceSignaled(const FD3D12Queue& Queue, FD3D12Payload* PendingPayload, uint64 CurrentFenceValue, uint64 LastCPUSignaledFenceValue);
+void ProcessInterruptQueueExit(uint8 Status);
+
+} }
 
 #endif
diff --git a/Engine/Source/Runtime/D3D12RHI/Private/Windows/WindowsD3D12Submission.h b/Engine/Source/Runtime/D3D12RHI/Private/Windows/WindowsD3D12Submission.h
--- a/Engine/Source/Runtime/D3D12RHI/Private/Windows/WindowsD3D12Submission.h
+++ b/Engine/Source/Runtime/D3D12RHI/Private/Windows/WindowsD3D12Submission.h
@@ -16,6 +16,6 @@ struct FD3D12Payload final : public FD3D12PayloadBase
 	FD3D12Payload(FD3D12Queue& Queue)
 		: FD3D12PayloadBase(Queue)
 	{
-		IF_RENDER_TRACING_ENABLED(RenderTracing::D3D12PayloadCreated(this, Queue));
+		IF_RENDER_TRACING_ENABLED(RenderTracing::D3D12::PayloadCreated(this, Queue));
 	}
 };
