D3D12 Descriptor Cache | Check for valid resource before allowing the SRV to be bound

diff --git a/Engine/Source/Runtime/D3D12RHI/Private/D3D12StateCache.cpp b/Engine/Source/Runtime/D3D12RHI/Private/D3D12StateCache.cpp
--- a/Engine/Source/Runtime/D3D12RHI/Private/D3D12StateCache.cpp
+++ b/Engine/Source/Runtime/D3D12RHI/Private/D3D12StateCache.cpp
@@ -1568,6 +1568,13 @@ void FD3D12StateCache::SetShaderResourceView(EShaderFrequency ShaderFrequency, F
 	FD3D12ShaderResourceViewCache& Cache = PipelineState.Common.SRVCache;
 	auto& CurrentShaderResourceViews = Cache.Views[ShaderFrequency];
 
+	// If we are binding an SRV that isn't actually backed by a resource, we shouldn't act as if we are non-null.
+	// Fall back to the 'null' binding path.
+	if (SRV && SRV->GetResource() == nullptr)
+	{
+		SRV = nullptr;
+	}
+
 	if ((CurrentShaderResourceViews[ResourceIndex] != SRV) || GD3D12SkipStateCaching)
 	{
 		if (SRV != nullptr)
