D3D12 Descriptors
- Free descriptor heap allocations back to the allocator during shutdown.
- Adding checks if heap allocations were not freed during shutdown.

#Synced-CL 52065068
#Synced-CL 52065068

diff --git a/Engine/Source/Runtime/D3D12RHI/Private/D3D12DescriptorCache.cpp b/Engine/Source/Runtime/D3D12RHI/Private/D3D12DescriptorCache.cpp
--- a/Engine/Source/Runtime/D3D12RHI/Private/D3D12DescriptorCache.cpp
+++ b/Engine/Source/Runtime/D3D12RHI/Private/D3D12DescriptorCache.cpp
@@ -904,6 +904,11 @@ FD3D12OnlineHeap::FD3D12OnlineHeap(FD3D12Device* Device, bool CanLoopAround)
 
 FD3D12OnlineHeap::~FD3D12OnlineHeap() = default;
 
+void FD3D12OnlineHeap::CleanupResources()
+{
+	Heap.SafeRelease();
+}
+
 /**
 Check if requested number of slots still fit the heap
 **/
diff --git a/Engine/Source/Runtime/D3D12RHI/Private/D3D12DescriptorCache.h b/Engine/Source/Runtime/D3D12RHI/Private/D3D12DescriptorCache.h
--- a/Engine/Source/Runtime/D3D12RHI/Private/D3D12DescriptorCache.h
+++ b/Engine/Source/Runtime/D3D12RHI/Private/D3D12DescriptorCache.h
@@ -180,6 +180,8 @@ public:
 	FD3D12OnlineHeap(FD3D12Device* Device, bool CanLoopAround);
 	virtual ~FD3D12OnlineHeap();
 
+	void CleanupResources();
+
 	ID3D12DescriptorHeap* GetHeap() { return Heap->GetHeap(); }
 
 	FORCEINLINE D3D12_CPU_DESCRIPTOR_HANDLE GetCPUSlotHandle(uint32 Slot) const { return Heap->GetCPUSlotHandle(Slot); }
diff --git a/Engine/Source/Runtime/D3D12RHI/Private/D3D12Descriptors.cpp b/Engine/Source/Runtime/D3D12RHI/Private/D3D12Descriptors.cpp
--- a/Engine/Source/Runtime/D3D12RHI/Private/D3D12Descriptors.cpp
+++ b/Engine/Source/Runtime/D3D12RHI/Private/D3D12Descriptors.cpp
@@ -137,8 +137,19 @@ FD3D12DescriptorHeap::FD3D12DescriptorHeap(FD3D12DescriptorHeap* SubAllocateSour
 {
 }
 
+FD3D12DescriptorHeap::FD3D12DescriptorHeap(FD3D12DescriptorManager* InManager, uint32 InOffset, uint32 InNumDescriptors)
+	: FD3D12DescriptorHeap(InManager->GetHeap(), InOffset, InNumDescriptors)
+{
+	Manager = InManager;
+}
+
 FD3D12DescriptorHeap::~FD3D12DescriptorHeap()
 {
+	if (Manager)
+	{
+		Manager->FreeSubheap(this);
+	}
+
 	if (CanBePooled())
 	{
 		GetParentDevice()->GetDescriptorHeapManager().AddHeapToPool(MoveTemp(Heap), GetType(), GetNumDescriptors(), GetFlags());
@@ -172,6 +183,21 @@ FD3D12DescriptorManager::FD3D12DescriptorManager(FD3D12Device* Device, FD3D12Des
 
 FD3D12DescriptorManager::~FD3D12DescriptorManager() = default;
 
+FD3D12DescriptorHeap* FD3D12DescriptorManager::AllocateSubheap(uint32 InNumDescriptors)
+{
+	if (TOptional<FRHIDescriptorAllocation> Allocation = Allocate(InNumDescriptors))
+	{
+		return new FD3D12DescriptorHeap(this, Allocation->StartIndex, Allocation->Count);
+	}
+
+	return nullptr;
+}
+
+void FD3D12DescriptorManager::FreeSubheap(FD3D12DescriptorHeap* InHeap)
+{
+	Free(FRHIDescriptorAllocation(InHeap->GetOffset(), InHeap->GetNumDescriptors()));
+}
+
 ////////////////////////////////////////////////////////////////////////////////////////////////////////////////////////
 // FD3D12OnlineDescriptorManager
 
@@ -180,7 +206,13 @@ FD3D12OnlineDescriptorManager::FD3D12OnlineDescriptorManager(FD3D12Device* Devic
 {
 }
 
-FD3D12OnlineDescriptorManager::~FD3D12OnlineDescriptorManager() = default;
+FD3D12OnlineDescriptorManager::~FD3D12OnlineDescriptorManager()
+{
+	for (const FD3D12DescriptorHeapPtr& Heap : Heaps)
+	{
+		checkf(!Heap, TEXT("FD3D12OnlineDescriptorManager destroyed without CleanupResources being called."));
+	}
+}
 
 // Allocate and initialize the online heap
 void FD3D12OnlineDescriptorManager::Init(uint32 InTotalSize, uint32 InBlockSize, bool bBindlessResources)
@@ -240,9 +272,9 @@ void FD3D12OnlineDescriptorManager::Init(uint32 InTotalSize, uint32 InBlockSize,
 
 void FD3D12OnlineDescriptorManager::CleanupResources()
 {
-	for (ERHIPipeline Pipeline : MakeFlagsRange(ERHIPipeline::All))
+	for (FD3D12DescriptorHeapPtr& Heap : Heaps)
 	{
-		Heaps[Pipeline] = nullptr;
+		Heap.SafeRelease();
 	}
 }
 
@@ -427,7 +459,7 @@ void FD3D12OfflineDescriptorManager::CleanupResources()
 {
 	for (FD3D12OfflineHeapEntry& HeapEntry : Heaps)
 	{
-		HeapEntry.Heap = nullptr;
+		HeapEntry.Heap.SafeRelease();
 	}
 }
 
@@ -540,9 +572,9 @@ FD3D12DescriptorHeap* FD3D12DescriptorHeapManager::AllocateHeap(const TCHAR* InD
 	{
 		if (GlobalHeap.HandlesHeapAllocationsWithFlags(InHeapType, InHeapFlags))
 		{
-			if (TOptional<FRHIDescriptorAllocation> Allocation = GlobalHeap.Allocate(InNumDescriptors))
+			if (FD3D12DescriptorHeap* NewHeap = GlobalHeap.AllocateSubheap(InNumDescriptors))
 			{
-				return new FD3D12DescriptorHeap(GlobalHeap.GetHeap(), Allocation->StartIndex, Allocation->Count);
+				return NewHeap;
 			}
 
 			// TODO: handle running out of space
@@ -555,43 +587,8 @@ FD3D12DescriptorHeap* FD3D12DescriptorHeapManager::AllocateHeap(const TCHAR* InD
 
 void FD3D12DescriptorHeapManager::DeferredFreeHeap(FD3D12DescriptorHeap* InHeap)
 {
-	if (InHeap->IsGlobal())
-	{
-		for (FD3D12DescriptorManager& GlobalHeap : GlobalHeaps)
-		{
-			if (GlobalHeap.IsHeapAChild(InHeap))
-			{
-				InHeap->AddRef();
-				FD3D12DynamicRHI::GetD3DRHI()->DeferredDelete(InHeap, FD3D12DeferredDeleteObject::EType::DescriptorHeap);
-				return;
-			}
-		}
-	}
-	else
-	{
-		InHeap->AddRef();
-		FD3D12DynamicRHI::GetD3DRHI()->DeferredDelete(InHeap, FD3D12DeferredDeleteObject::EType::DescriptorHeap);
-	}
-}
-
-void FD3D12DescriptorHeapManager::ImmediateFreeHeap(FD3D12DescriptorHeap* InHeap)
-{
-	if (InHeap->IsGlobal())
-	{
-		for (FD3D12DescriptorManager& GlobalHeap : GlobalHeaps)
-		{
-			if (GlobalHeap.IsHeapAChild(InHeap))
-			{
-				GlobalHeap.Free(FRHIDescriptorAllocation(InHeap->GetOffset(), InHeap->GetNumDescriptors()));
-				InHeap->Release();
-				return;
-			}
-		}
-	}
-	else
-	{
-		InHeap->Release();
-	}
+	InHeap->AddRef();
+	FD3D12DynamicRHI::GetD3DRHI()->DeferredDelete(InHeap, FD3D12DeferredDeleteObject::EType::DescriptorHeap);
 }
 
 void FD3D12DescriptorHeapManager::AddHeapToPool(TRefCountPtr<ID3D12DescriptorHeap>&& InHeap, ERHIDescriptorHeapType InType, uint32 InNumDescriptors, ED3D12DescriptorHeapFlags InFlags)
diff --git a/Engine/Source/Runtime/D3D12RHI/Private/D3D12Descriptors.h b/Engine/Source/Runtime/D3D12RHI/Private/D3D12Descriptors.h
--- a/Engine/Source/Runtime/D3D12RHI/Private/D3D12Descriptors.h
+++ b/Engine/Source/Runtime/D3D12RHI/Private/D3D12Descriptors.h
@@ -11,6 +11,7 @@
 #include "RHIPipeline.h"
 #include "Templates/RefCounting.h"
 
+class FD3D12DescriptorManager;
 class FD3D12View;
 
 struct FD3D12DescriptorHeap;
@@ -106,6 +107,9 @@ public:
 	// Heap created as a suballocation of another heap.
 	FD3D12DescriptorHeap(FD3D12DescriptorHeap* SubAllocateSourceHeap, uint32 InOffset, uint32 InNumDescriptors);
 
+	// Heap created as an allocation from a descriptor manager.
+	FD3D12DescriptorHeap(FD3D12DescriptorManager* Manager, uint32 InOffset, uint32 InNumDescriptors);
+
 	~FD3D12DescriptorHeap();
 
 	ID3D12DescriptorHeap*       GetHeap()  const { return Heap; }
@@ -133,6 +137,9 @@ public:
 private:
 	TRefCountPtr<ID3D12DescriptorHeap> Heap;
 
+	// Manager this heap is a suballocation of, if any.
+	FD3D12DescriptorManager* Manager = nullptr;
+
 	const CD3DX12_CPU_DESCRIPTOR_HANDLE CpuBase;
 	const CD3DX12_GPU_DESCRIPTOR_HANDLE GpuBase;
 
@@ -167,6 +174,9 @@ public:
 	inline       FD3D12DescriptorHeap* GetHeap()       { return Heap.GetReference(); }
 	inline const FD3D12DescriptorHeap* GetHeap() const { return Heap.GetReference(); }
 
+	FD3D12DescriptorHeap* AllocateSubheap(uint32 InNumDescriptors);
+	void FreeSubheap(FD3D12DescriptorHeap* InHeap);
+
 	bool HandlesHeapAllocationsWithFlags(ERHIDescriptorHeapType InHeapType, ED3D12DescriptorHeapFlags InHeapFlags) const
 	{
 		return HandlesAllocations(D3D12DescriptorTypeMaskFromHeapType(InHeapType)) && Heap->GetFlags() == InHeapFlags;
@@ -177,7 +187,6 @@ public:
 		return HandlesAllocation(InDescriptorType) && Heap->GetFlags() == InHeapFlags;
 	}
 
-
 	inline bool IsHeapAChild(const FD3D12DescriptorHeap* InHeap)
 	{
 		return InHeap->GetHeap() == Heap->GetHeap();
@@ -311,7 +320,6 @@ public:
 	FD3D12DescriptorHeap* AllocateIndependentHeap(const TCHAR* InDebugName, ERHIDescriptorHeapType InHeapType, uint32 InNumDescriptors, ED3D12DescriptorHeapFlags InHeapFlags);
 	FD3D12DescriptorHeap* AllocateHeap(const TCHAR* InDebugName, ERHIDescriptorHeapType InHeapType, uint32 InNumDescriptors, ED3D12DescriptorHeapFlags InHeapFlags);
 	void DeferredFreeHeap(FD3D12DescriptorHeap* InHeap);
-	void ImmediateFreeHeap(FD3D12DescriptorHeap* InHeap);
 
 	void AddHeapToPool(TRefCountPtr<ID3D12DescriptorHeap>&& InHeap, ERHIDescriptorHeapType InType, uint32 InNumDescriptors, ED3D12DescriptorHeapFlags InFlags);
 
diff --git a/Engine/Source/Runtime/D3D12RHI/Private/D3D12Device.cpp b/Engine/Source/Runtime/D3D12RHI/Private/D3D12Device.cpp
--- a/Engine/Source/Runtime/D3D12RHI/Private/D3D12Device.cpp
+++ b/Engine/Source/Runtime/D3D12RHI/Private/D3D12Device.cpp
@@ -438,6 +438,7 @@ void FD3D12Device::CleanupResources()
 		Manager.CleanupResources();
 	}
 	OnlineDescriptorManager.CleanupResources();
+	GlobalSamplerHeap.CleanupResources();
 
 #if PLATFORM_SUPPORTS_BINDLESS_RENDERING
 	BindlessDescriptorManager.CleanupResources();
diff --git a/Engine/Source/Runtime/D3D12RHI/Private/D3D12RHI.cpp b/Engine/Source/Runtime/D3D12RHI/Private/D3D12RHI.cpp
--- a/Engine/Source/Runtime/D3D12RHI/Private/D3D12RHI.cpp
+++ b/Engine/Source/Runtime/D3D12RHI/Private/D3D12RHI.cpp
@@ -480,7 +480,7 @@ void FD3D12DynamicRHI::RHIProcessDeleteQueue()
 					ObjectToDelete.Heap->Release();
 					break;
 				case FD3D12DeferredDeleteObject::EType::DescriptorHeap:
-					ObjectToDelete.DescriptorHeap->GetParentDevice()->GetDescriptorHeapManager().ImmediateFreeHeap(ObjectToDelete.DescriptorHeap);
+					ObjectToDelete.DescriptorHeap->Release();
 					break;
 				case FD3D12DeferredDeleteObject::EType::D3DObject:
 					ObjectToDelete.D3DObject->Release();
