D3D12 Bindless update optimization
- When we collect bindless indices that need updating, we aren't keeping the update list unique.
- In some scenarios, especially in the editor, we can end up with an update count that is 10x to 1000x the number of total descriptors in the heap.
- Adding a new threshold cvar, D3D12.Bindless.HeapUpdateUniqueUploadThreshold. When the number of updates are over that threshold, we will force the update list to be unique.

#Synced-CL 52065068

diff --git a/Engine/Source/Runtime/D3D12RHI/Private/D3D12BindlessDescriptors.cpp b/Engine/Source/Runtime/D3D12RHI/Private/D3D12BindlessDescriptors.cpp
--- a/Engine/Source/Runtime/D3D12RHI/Private/D3D12BindlessDescriptors.cpp
+++ b/Engine/Source/Runtime/D3D12RHI/Private/D3D12BindlessDescriptors.cpp
@@ -3,6 +3,8 @@
 #include "D3D12BindlessDescriptors.h"
 #include "D3D12RHIPrivate.h"
 #include "D3D12Descriptors.h"
+#include "Algo/Sort.h"
+#include "Algo/Unique.h"
 
 #if PLATFORM_SUPPORTS_BINDLESS_RENDERING
 
@@ -29,6 +31,14 @@ static FAutoConsoleVariableRef CVarBindlessSamplerDescriptorHeapSize(
 	ECVF_ReadOnly
 );
 
+int32 GBindlessHeapUpdateUniqueUploadThreshold = 1024;
+static FAutoConsoleVariableRef CVarBindlessHeapUpdateUniqueUploadThreshold(
+	TEXT("D3D12.Bindless.HeapUpdateUniqueUploadThreshold"),
+	GBindlessHeapUpdateUniqueUploadThreshold,
+	TEXT("Threshold of upload index count before we make the list of update indices unique."),
+	ECVF_Default
+);
+
 FD3D12DescriptorHeap* UE::D3D12BindlessDescriptors::CreateCpuHeap(FD3D12Device* InDevice, ERHIDescriptorHeapType InType, uint32 InNewNumDescriptorsPerHeap)
 {
 	LLM_SCOPE_BYNAME(TEXT("RHIMisc/BindlessDescriptorHeap/CPU"));
@@ -745,12 +755,23 @@ void FD3D12BindlessResourceManager::CheckRequestNewActiveGPUHeap()
 	// Found a pooled heap, then copy over the dirty descriptor handles
 	if (NewActiveGpuHeapIndex >= 0)
 	{
-		// NOTE: copying over duplicate descriptor entries is faster then adding them to set for reduction
-		//		 CitySample there is about 2 to 4 times duplication but still faster to copy all then deduplication
+		// NOTE: For small update lists, copying duplicates is faster than deduplicating.
+		//       For large lists (see GBindlessHeapUpdateUniqueUploadThreshold), we sort + unique
+		//       to avoid excessive driver overhead from large CopyDescriptors calls.
 
 		FGpuHeapData& GpuHeapData = ActiveGpuHeaps[NewActiveGpuHeapIndex];
 		INC_DWORD_STAT_BY(STAT_D3D12BindlessResourceGPUDescriptorsCopied, GpuHeapData.UpdatedIndices.Num());
 
+		// Deduplicate the updated indices when the list is large enough that sorting + unique
+		// is cheaper than copying redundant descriptors. The first condition catches pathological
+		// cases where accumulated updates exceed the heap size (guaranteed duplicates).
+		if (GpuHeapData.UpdatedIndices.Num() > (int32)GpuHeapData.GpuHeap->GetNumDescriptors() ||
+			GpuHeapData.UpdatedIndices.Num() > GBindlessHeapUpdateUniqueUploadThreshold)
+		{
+			Algo::Sort(GpuHeapData.UpdatedIndices);
+			GpuHeapData.UpdatedIndices.SetNum(Algo::Unique(GpuHeapData.UpdatedIndices));
+		}
+
 		UE::D3D12Descriptors::CopyDescriptors(GetParentDevice(), GpuHeapData.GpuHeap, CpuHeap, GpuHeapData.UpdatedIndices);
 		GpuHeapData.UpdatedIndices.Reset();
 
