- Make d3d12 rhi work without engine (add WITH_ENGINE add 2 required locations)
- Add BuildComputeShaderRootSignature() in to internal header so it can use the d3d12 rhi code to setup root signatures matching the engine logic in the ASDTool

#Synced-CL 52126265

diff --git a/Engine/Source/Runtime/D3D12RHI/Internal/D3D12ShaderUtils.h b/Engine/Source/Runtime/D3D12RHI/Internal/D3D12ShaderUtils.h
new file mode 100644
--- /dev/null
+++ b/Engine/Source/Runtime/D3D12RHI/Internal/D3D12ShaderUtils.h
@@ -0,0 +1,22 @@
+// Copyright Epic Games, Inc. All Rights Reserved.
+
+#pragma once
+
+#include "CoreMinimal.h"
+#include "D3D12ThirdParty.h"
+
+/**
+ * Build a serialized root signature for a compute shader from its UE shader blob.
+ * Uses the same root signature construction logic as RHICreateComputeShader.
+ *
+ * @param ShaderCode            Full UE shader blob (SRT + native bytecode + optional data)
+ * @param ResourceBindingTier   Resource binding tier to use for the root signature layout
+ * @param OutRootSigBlob        Receives the serialized root signature blob
+ * @param OutNativeBytecodeOffset  Receives the offset to native DXBC/DXIL bytecode within ShaderCode
+ * @return true on success
+ */
+D3D12RHI_API bool BuildComputeShaderRootSignature(
+	const TConstArrayView<uint8>& ShaderCode,
+	D3D12_RESOURCE_BINDING_TIER ResourceBindingTier,
+	TRefCountPtr<ID3DBlob>& OutRootSigBlob,
+	int32& OutNativeBytecodeOffset);
diff --git a/Engine/Source/Runtime/D3D12RHI/Private/D3D12RHIPrivate.h b/Engine/Source/Runtime/D3D12RHI/Private/D3D12RHIPrivate.h
--- a/Engine/Source/Runtime/D3D12RHI/Private/D3D12RHIPrivate.h
+++ b/Engine/Source/Runtime/D3D12RHI/Private/D3D12RHIPrivate.h
@@ -21,7 +21,6 @@
 #include "HAL/ThreadManager.h"
 #include "Containers/ResourceArray.h"
 #include "Serialization/MemoryReader.h"
-#include "EngineGlobals.h"
 #include "StaticBoundShaderState.h"
 
 // Dependencies.
diff --git a/Engine/Source/Runtime/D3D12RHI/Private/D3D12RenderTarget.cpp b/Engine/Source/Runtime/D3D12RHI/Private/D3D12RenderTarget.cpp
--- a/Engine/Source/Runtime/D3D12RHI/Private/D3D12RenderTarget.cpp
+++ b/Engine/Source/Runtime/D3D12RHI/Private/D3D12RenderTarget.cpp
@@ -5,11 +5,8 @@
 	=============================================================================*/
 
 #include "D3D12RHIPrivate.h"
-#include "BatchedElements.h"
-#include "ScreenRendering.h"
 #include "RHIStaticStates.h"
 #include "ResolveShader.h"
-#include "SceneUtils.h"
 #include "PipelineStateCache.h"
 #include "Math/PackedVector.h"
 #include "RHISurfaceDataConversion.h"
diff --git a/Engine/Source/Runtime/D3D12RHI/Private/D3D12Resources.cpp b/Engine/Source/Runtime/D3D12RHI/Private/D3D12Resources.cpp
--- a/Engine/Source/Runtime/D3D12RHI/Private/D3D12Resources.cpp
+++ b/Engine/Source/Runtime/D3D12RHI/Private/D3D12Resources.cpp
@@ -8,7 +8,6 @@ D3D12Resources.cpp: D3D RHI utility implementation.
 #include "D3D12RHIPrivate.h"
 #include "D3D12IntelExtensions.h"
 #include "D3D12RayTracing.h"
-#include "EngineModule.h"
 #include "HAL/LowLevelMemTracker.h"
 #include "ProfilingDebugging/MemoryTrace.h"
 #include "ProfilingDebugging/AssetMetadataTrace.h"
diff --git a/Engine/Source/Runtime/D3D12RHI/Private/D3D12Stats.cpp b/Engine/Source/Runtime/D3D12RHI/Private/D3D12Stats.cpp
--- a/Engine/Source/Runtime/D3D12RHI/Private/D3D12Stats.cpp
+++ b/Engine/Source/Runtime/D3D12RHI/Private/D3D12Stats.cpp
@@ -6,7 +6,6 @@ D3D12Stats.cpp:RHI Stats and timing implementation.
 
 #include "D3D12Stats.h"
 #include "D3D12RHIPrivate.h"
-#include "Engine/Engine.h"
 #include "ProfilingDebugging/MemoryTrace.h"
 #include "RHICoreStats.h"
 
diff --git a/Engine/Source/Runtime/D3D12RHI/Private/D3D12Util.cpp b/Engine/Source/Runtime/D3D12RHI/Private/D3D12Util.cpp
--- a/Engine/Source/Runtime/D3D12RHI/Private/D3D12Util.cpp
+++ b/Engine/Source/Runtime/D3D12RHI/Private/D3D12Util.cpp
@@ -6,10 +6,13 @@ D3D12Util.h: D3D RHI utility implementation.
 
 #include "D3D12Util.h"
 #include "D3D12RHIPrivate.h"
-#include "EngineModule.h"
+#include "D3D12ShaderUtils.h"
 #include "RendererInterface.h"
 #include "CoreGlobals.h"
+#if WITH_ENGINE
 #include "Engine/GameEngine.h"
+#include "EngineModule.h"
+#endif
 #include "Misc/OutputDeviceRedirector.h"
 #include "HAL/ExceptionHandling.h"
 #if PLATFORM_WINDOWS
@@ -790,7 +793,9 @@ void FD3D12DynamicRHI::TerminateOnOutOfMemory(ID3D12Device* InDevice, HRESULT D3
 	}
 
 #if STATS
+#if WITH_ENGINE
 	GetRendererModule().DebugLogOnCrash();
+#endif
 #endif
 
 	static IConsoleVariable* GPUCrashOOM = IConsoleManager::Get().FindConsoleVariable(TEXT("r.GPUCrashOnOutOfMemory"));
@@ -1009,11 +1014,13 @@ void FD3D12DynamicRHI::TerminateOnGPUCrash(TOptional<HRESULT> D3DResult)
 	FTextBuilder ErrorMessage;
 	OutputGPUCrashReport(ErrorMessage);
 
+#if WITH_ENGINE
 	UGameEngine* GameEngine = Cast<UGameEngine>(GEngine);
 	if (GameEngine)
 	{
 		GameEngine->OnGPUCrash();
 	}
+#endif
 
 	// Show message box or trace information
 #if !(UE_BUILD_SHIPPING || UE_BUILD_TEST)
@@ -1925,4 +1932,82 @@ DEFINE_STAT(STAT_ExplicitMaxUsedSamplerDescriptors);
 DEFINE_STAT(STAT_ExplicitUsedSamplerDescriptors);
 DEFINE_STAT(STAT_ExplicitUsedViewDescriptors);
 
+
+//--------------------------------------------------------------------------------------------------
+// BuildComputeShaderRootSignature — Build a serialized root signature from a UE shader blob
+// Uses the same root signature construction logic as RHICreateComputeShader.
+// Can be called without an initialized D3D12 RHI.
+//--------------------------------------------------------------------------------------------------
+
+bool BuildComputeShaderRootSignature(const TConstArrayView<uint8>& ShaderCode, D3D12_RESOURCE_BINDING_TIER ResourceBindingTier, TRefCountPtr<ID3DBlob>& OutRootSigBlob, int32& OutNativeBytecodeOffset)
+{
+	if (ShaderCode.Num() == 0)
+	{
+		return false;
+	}
+
+	// Parse the shader blob - same as InitStandardShader
+	FShaderCodeReader CodeReader(ShaderCode);
+
+	// Extract resource counts from optional data
+	const FShaderCodePackedResourceCounts* PackedResourceCounts = CodeReader.FindOptionalData<FShaderCodePackedResourceCounts>();
+	if (!PackedResourceCounts)
+	{
+		UE_LOG(LogD3D12RHI, Warning, TEXT("BuildComputeShaderRootSignature: No resource counts found in shader blob"));
+		return false;
+	}
+
+	// Deserialize SRT to find where native bytecode starts
+	FMemoryReaderView Ar(ShaderCode, true);
+	FShaderResourceTable SRT;
+	Ar << SRT;
+	OutNativeBytecodeOffset = Ar.Tell();
+
+	// Build quantized bound shader state - same as GetRootSignature(ComputeShader)
+	FD3D12QuantizedBoundShaderState QBSS{};
+
+	// Use a temporary FD3D12ShaderData to hold the extracted data
+	// so we can call the existing quantization functions
+	FD3D12ShaderData TempShaderData;
+	TempShaderData.ResourceCounts = *PackedResourceCounts;
+
+#if D3D12RHI_NEEDS_VENDOR_EXTENSIONS
+	// Extract vendor extensions from optional data
+	int32 VendorExtensionTableSize = 0;
+	auto* VendorExtensionData = CodeReader.FindOptionalDataAndSize(FShaderCodeVendorExtension::Key, VendorExtensionTableSize);
+	if (VendorExtensionData && VendorExtensionTableSize > 0)
+	{
+		FBufferReader VendorAr((void*)VendorExtensionData, VendorExtensionTableSize, false);
+		VendorAr << TempShaderData.VendorExtensions;
+	}
+#endif
+
+	// Root constants are always supported on D3D12 - force enable for tool context
+	// where GRHISupportsShaderRootConstants is not initialized by device init
+	GRHISupportsShaderRootConstants = true;
+
+	// Quantize for compute: SV_All visibility, UAVs allowed
+	QuantizeBoundShaderStateCommon(QBSS, &TempShaderData, ResourceBindingTier, SV_All, true /*bAllowUAVs*/);
+
+	// Build root signature descriptor
+	FD3D12RootSignatureDesc RootSigDesc(QBSS, ResourceBindingTier);
+
+	// Serialize root signature
+	TRefCountPtr<ID3DBlob> ErrorBlob;
+	HRESULT hr = D3DX12SerializeVersionedRootSignature(
+		&RootSigDesc.GetDesc(),
+		D3D_ROOT_SIGNATURE_VERSION_1_1,
+		OutRootSigBlob.GetInitReference(),
+		ErrorBlob.GetInitReference());
+
+	if (FAILED(hr))
+	{
+		UE_LOG(LogD3D12RHI, Error, TEXT("BuildComputeShaderRootSignature: Failed to serialize root signature: 0x%08X %s"),
+			hr, ErrorBlob ? ANSI_TO_TCHAR((const char*)ErrorBlob->GetBufferPointer()) : TEXT(""));
+		return false;
+	}
+
+	return true;
+}
+
 #undef LOCTEXT_NAMESPACE
diff --git a/Engine/Source/Runtime/D3D12RHI/Private/D3D12Viewport.cpp b/Engine/Source/Runtime/D3D12RHI/Private/D3D12Viewport.cpp
--- a/Engine/Source/Runtime/D3D12RHI/Private/D3D12Viewport.cpp
+++ b/Engine/Source/Runtime/D3D12RHI/Private/D3D12Viewport.cpp
@@ -7,7 +7,6 @@
 #include "D3D12Viewport.h"
 #include "D3D12RHIPrivate.h"
 #include "RenderCore.h"
-#include "Engine/RendererSettings.h"
 #include "HDRHelper.h"
 #include "DataDrivenShaderPlatformInfo.h"
 #include "RHIUtilities.h"
diff --git a/Engine/Source/Runtime/D3D12RHI/Private/Windows/WindowsD3D12Device.cpp b/Engine/Source/Runtime/D3D12RHI/Private/Windows/WindowsD3D12Device.cpp
--- a/Engine/Source/Runtime/D3D12RHI/Private/Windows/WindowsD3D12Device.cpp
+++ b/Engine/Source/Runtime/D3D12RHI/Private/Windows/WindowsD3D12Device.cpp
@@ -13,12 +13,14 @@
 #include "Windows/WindowsPlatformCrashContext.h"
 #include "HAL/FileManager.h"
 #include "ProfilingDebugging/AssetMetadataTrace.h"
+#if WITH_ENGINE
 #include "HardwareInfo.h"
 #include "IHeadMountedDisplayModule.h"
+#endif // WITH_ENGINE
 #include "GenericPlatform/GenericPlatformDriver.h"			// FGPUDriverInfo
 #include "RHIValidation.h"
 #include "RHIUtilities.h"
-#include "ShaderCompiler.h"
+
 #include "Misc/EngineVersion.h"
 #include "Windows/WindowsD3D.h"
 
@@ -687,13 +689,15 @@ bool FD3D12DynamicRHI::SetupDisplayHDRMetaData()
 }
 
 #if PLATFORM_WINDOWS
-extern void HDRSettingChangedSinkCallback();
 void FD3D12DynamicRHI::RHIHandleDisplayChange()
 {
 	RHIBlockUntilGPUIdle();
 	GRHISupportsHDROutput = SetupDisplayHDRMetaData();
+#if WITH_ENGINE
 	// make sure CVars are being updated properly
+	extern void HDRSettingChangedSinkCallback();
 	HDRSettingChangedSinkCallback();
+#endif
 }
 #endif
 
@@ -884,7 +888,11 @@ void FD3D12DynamicRHIModule::FindAdapter()
 	DXGIFactory4->QueryInterface(IID_PPV_ARGS(DXGIFactory6.GetInitReference()));
 
 	// Allow HMD to override which graphics adapter is chosen, so we pick the adapter where the HMD is connected
+#if WITH_ENGINE
 	const uint64 HmdGraphicsAdapterLuid = IHeadMountedDisplayModule::IsAvailable() ? IHeadMountedDisplayModule::Get().GetGraphicsAdapterLuid() : 0;
+#else
+	const uint64 HmdGraphicsAdapterLuid = 0;
+#endif
 	const uint64 AdapterLuid = FWindowsD3D::ChooseD3D12Adapter(HmdGraphicsAdapterLuid);
 
 	TRefCountPtr<IDXGIAdapter> TempAdapter;
@@ -1970,7 +1978,9 @@ void FD3D12DynamicRHI::Init()
 #endif
 	}
 
+#if WITH_ENGINE
 	FHardwareInfo::RegisterHardwareInfo(NAME_RHI, TEXT("D3D12"));
+#endif
 
 	GRHISupportsTextureStreaming = true;
 	GRHISupportsFirstInstance = true;
