Fix missing UpdateResidency for fast allocator page in ray tracing dispatch

#Synced-CL 52149918

diff --git a/Engine/Source/Runtime/D3D12RHI/Private/D3D12RayTracing.cpp b/Engine/Source/Runtime/D3D12RHI/Private/D3D12RayTracing.cpp
--- a/Engine/Source/Runtime/D3D12RHI/Private/D3D12RayTracing.cpp
+++ b/Engine/Source/Runtime/D3D12RHI/Private/D3D12RayTracing.cpp
@@ -1804,7 +1804,7 @@ public:
 		return Buffer->ResourceLocation.GetGPUVirtualAddress();
 	}
 
-	D3D12_DISPATCH_RAYS_DESC GetDispatchRaysDesc(FD3D12Device* InDevice, const FD3D12ShaderIdentifier& RayGenShaderIdentifier) const
+	D3D12_DISPATCH_RAYS_DESC GetDispatchRaysDesc(FD3D12Device* InDevice, const FD3D12ShaderIdentifier& RayGenShaderIdentifier, FD3D12Resource*& OutRayGenUploadResource) const
 	{
 		check(Device == InDevice);
 
@@ -1814,6 +1814,7 @@ public:
 		void* RayGenGPUData = Device->GetDefaultFastAllocator().Allocate(RayGenRecordStride, 256, &UploadResourceLocation);
 		FMemory::Memcpy(RayGenGPUData, &RayGenShaderIdentifier, ShaderIdentifierSize);
 		D3D12_GPU_VIRTUAL_ADDRESS RayGenStartShaderIdentifierAddress = UploadResourceLocation.GetGPUVirtualAddress();
+		OutRayGenUploadResource = UploadResourceLocation.GetResource();
 		
 		D3D12_GPU_VIRTUAL_ADDRESS ShaderTableAddress = GetShaderTableAddress();
 
@@ -5815,6 +5816,7 @@ static void DispatchRays(FD3D12CommandContext& CommandContext,
 	FD3D12RayTracingShaderBindingTableInternal* OptShaderTable,
 	const D3D12_DISPATCH_RAYS_DESC& DispatchDesc,
 	ED3D12QueueType QueueType,
+	FD3D12Resource* RayGenUploadResource = nullptr,
 	FD3D12Buffer* ArgumentBuffer = nullptr, uint32 ArgumentOffset = 0)
 {
 	SCOPE_CYCLE_COUNTER(STAT_D3D12DispatchRays);
@@ -5968,6 +5970,12 @@ static void DispatchRays(FD3D12CommandContext& CommandContext,
 
 	if (bResourcesBound)
 	{
+		// Ensure residency for the ray gen
+		if (RayGenUploadResource)
+		{
+			CommandContext.UpdateResidency(RayGenUploadResource);
+		}
+
 		CommandContext.FlushResourceBarriers();
 
 		ID3D12StateObject* RayTracingStateObject = nullptr;
@@ -6036,13 +6044,14 @@ void FD3D12CommandContext::RHIRayTraceDispatch(FRHIRayTracingPipelineState* InRa
 			*(RayGenShader->EntryPoint));
 
 	const FD3D12ShaderIdentifier& RayGenShaderIdentifier = Pipeline->RayGenShaders.Identifiers[RayGenShaderIndex];
-	D3D12_DISPATCH_RAYS_DESC DispatchDesc = ShaderTableForDevice->GetDispatchRaysDesc(GetParentDevice(), RayGenShaderIdentifier);
+	FD3D12Resource* RayGenUploadResource = nullptr;
+	D3D12_DISPATCH_RAYS_DESC DispatchDesc = ShaderTableForDevice->GetDispatchRaysDesc(GetParentDevice(), RayGenShaderIdentifier, RayGenUploadResource);
 
 	DispatchDesc.Width = Width;
 	DispatchDesc.Height = Height;
 	DispatchDesc.Depth = 1;
 
-	DispatchRays(*this, GlobalResourceBindings, Pipeline, RayGenShaderIndex, ShaderTableForDevice, DispatchDesc, QueueType);
+	DispatchRays(*this, GlobalResourceBindings, Pipeline, RayGenShaderIndex, ShaderTableForDevice, DispatchDesc, QueueType, RayGenUploadResource);
 }
 
 void FD3D12CommandContext::RHIRayTraceDispatchIndirect(FRHIRayTracingPipelineState* InRayTracingPipelineState, FRHIRayTracingShader* RayGenShaderRHI,
@@ -6063,8 +6072,9 @@ void FD3D12CommandContext::RHIRayTraceDispatchIndirect(FRHIRayTracingPipelineSta
 	checkf(RayGenShaderIndex != INDEX_NONE, TEXT("RayGen shader is not present in the given ray tracing pipeline. All RayGen shaders must be declared when creating RTPSO."));
 
 	const FD3D12ShaderIdentifier& RayGenShaderIdentifier = Pipeline->RayGenShaders.Identifiers[RayGenShaderIndex];
-	D3D12_DISPATCH_RAYS_DESC DispatchDesc = ShaderTableForDevice->GetDispatchRaysDesc(GetParentDevice(), RayGenShaderIdentifier);
-	DispatchRays(*this, GlobalResourceBindings, Pipeline, RayGenShaderIndex, ShaderTableForDevice, DispatchDesc, QueueType, RetrieveObject<FD3D12Buffer>(ArgumentBuffer), ArgumentOffset);
+	FD3D12Resource* RayGenUploadResource = nullptr;
+	D3D12_DISPATCH_RAYS_DESC DispatchDesc = ShaderTableForDevice->GetDispatchRaysDesc(GetParentDevice(), RayGenShaderIdentifier, RayGenUploadResource);
+	DispatchRays(*this, GlobalResourceBindings, Pipeline, RayGenShaderIndex, ShaderTableForDevice, DispatchDesc, QueueType, RayGenUploadResource, RetrieveObject<FD3D12Buffer>(ArgumentBuffer), ArgumentOffset);
 }
 
 static void SetRayTracingHitGroup(
