Add more CPU profiling scopes to major functions that run on editor/engine startup

#rnx
#Synced-CL 52172106

diff --git a/Engine/Plugins/Runtime/WindowsDeviceProfileSelector/Source/WindowsDeviceProfileSelector/Private/WindowsDeviceProfileSelectorModule.cpp b/Engine/Plugins/Runtime/WindowsDeviceProfileSelector/Source/WindowsDeviceProfileSelector/Private/WindowsDeviceProfileSelectorModule.cpp
--- a/Engine/Plugins/Runtime/WindowsDeviceProfileSelector/Source/WindowsDeviceProfileSelector/Private/WindowsDeviceProfileSelectorModule.cpp
+++ b/Engine/Plugins/Runtime/WindowsDeviceProfileSelector/Source/WindowsDeviceProfileSelector/Private/WindowsDeviceProfileSelectorModule.cpp
@@ -22,6 +22,8 @@ void FWindowsDeviceProfileSelectorModule::ShutdownModule()
 
 const FString FWindowsDeviceProfileSelectorModule::GetRuntimeDeviceProfileName()
 {
+	TRACE_CPUPROFILER_EVENT_SCOPE(FWindowsDeviceProfileSelectorModule::GetRuntimeDeviceProfileName);
+
 	static FString ProfileName;
 
 	if (ProfileName.IsEmpty())
diff --git a/Engine/Source/Developer/DerivedDataCache/Private/DerivedDataBackends.cpp b/Engine/Source/Developer/DerivedDataCache/Private/DerivedDataBackends.cpp
--- a/Engine/Source/Developer/DerivedDataCache/Private/DerivedDataBackends.cpp
+++ b/Engine/Source/Developer/DerivedDataCache/Private/DerivedDataBackends.cpp
@@ -316,6 +316,8 @@ private:
 	 */
 	bool ParseGraph(FStringView GraphNameOrConfig, FParsedNodeMap& InParsedNodes)
 	{
+		TRACE_CPUPROFILER_EVENT_SCOPE(FDerivedDataBackendGraph::ParseGraph);
+
 		RootCache = CreateCacheStoreHierarchy(Hierarchy, [this](IMemoryCacheStore*& OutCache) { GetMemoryCache(OutCache); });
 
 		TStringBuilder<256> GraphConfig;
diff --git a/Engine/Source/Runtime/D3D12RHI/Private/D3D12Adapter.cpp b/Engine/Source/Runtime/D3D12RHI/Private/D3D12Adapter.cpp
--- a/Engine/Source/Runtime/D3D12RHI/Private/D3D12Adapter.cpp
+++ b/Engine/Source/Runtime/D3D12RHI/Private/D3D12Adapter.cpp
@@ -494,6 +494,8 @@ FD3D12Adapter::FD3D12Adapter(FD3D12AdapterDesc& DescIn)
 
 void FD3D12Adapter::CreateRootDevice(bool bWithDebug)
 {
+	TRACE_CPUPROFILER_EVENT_SCOPE(FD3D12Adapter::CreateRootDevice);
+
 	// -d3ddebug is always allowed on Windows, but only allowed in non-shipping builds on other platforms.
 	// -gpuvalidation is only supported on Windows.
 #if PLATFORM_WINDOWS || !UE_BUILD_SHIPPING
@@ -745,11 +747,14 @@ void FD3D12Adapter::CreateRootDevice(bool bWithDebug)
 #endif
 
 		// Creating the Direct3D device.
-		VERIFYD3D12RESULT(D3D12CreateDevice(
-			GetAdapter(),
-			GetFeatureLevel(),
-			IID_PPV_ARGS(RootDevice.GetInitReference())
-		));
+		{
+			TRACE_CPUPROFILER_EVENT_SCOPE(D3D12CreateDevice);
+			VERIFYD3D12RESULT(D3D12CreateDevice(
+				GetAdapter(),
+				GetFeatureLevel(),
+				IID_PPV_ARGS(RootDevice.GetInitReference())
+			));
+		}
 	}
 
 #if ENABLE_RESIDENCY_MANAGEMENT
@@ -1012,6 +1017,8 @@ FD3D12TransientHeapCache& FD3D12Adapter::GetOrCreateTransientHeapCache()
 
 void FD3D12Adapter::InitializeDevices()
 {
+	TRACE_CPUPROFILER_EVENT_SCOPE(FD3D12Adapter::InitializeDevices);
+
 	check(IsInGameThread());
 
 	// Wait for the rendering thread to go idle.
diff --git a/Engine/Source/Runtime/D3D12RHI/Private/Windows/WindowsD3D12Device.cpp b/Engine/Source/Runtime/D3D12RHI/Private/Windows/WindowsD3D12Device.cpp
--- a/Engine/Source/Runtime/D3D12RHI/Private/Windows/WindowsD3D12Device.cpp
+++ b/Engine/Source/Runtime/D3D12RHI/Private/Windows/WindowsD3D12Device.cpp
@@ -1430,6 +1430,8 @@ void FD3D12DynamicRHI::InitializeDeviceAgnosticVendorAPI()
 
 void FD3D12DynamicRHI::Init()
 {
+	TRACE_CPUPROFILER_EVENT_SCOPE(FD3D12DynamicRHI::Init);
+
 #if D3D12RHI_SUPPORTS_WIN_PIX
 	// PIX GPU capture dll: makes PIX be able to attach to our process. !GetModuleHandle() is required because PIX may already be attached.
 	if (GAutoAttachPIX || FParse::Param(FCommandLine::Get(), TEXT("attachPIX")))
@@ -1663,15 +1665,19 @@ void FD3D12DynamicRHI::Init()
 #endif
 				if (!bSkipSERSupportCheck)
 				{
-					NvAPI_Status NvStatusSER = NvAPI_D3D12_IsNvShaderExtnOpCodeSupported(GetAdapter().GetD3DDevice(), NV_EXTN_OP_HIT_OBJECT_REORDER_THREAD, &GRHIGlobals.SupportsShaderExecutionReordering);
+					ID3D12Device* D3DDevice = GetAdapter().GetD3DDevice();
+					NvAPI_Status NvStatusSER = NvAPI_D3D12_IsNvShaderExtnOpCodeSupported(D3DDevice, NV_EXTN_OP_HIT_OBJECT_REORDER_THREAD, &GRHIGlobals.SupportsShaderExecutionReordering);
 					if (NvStatusSER == NVAPI_OK && GRHIGlobals.SupportsShaderExecutionReordering)
 					{
 						NVAPI_D3D12_RAYTRACING_THREAD_REORDERING_CAPS ReorderCaps = NVAPI_D3D12_RAYTRACING_THREAD_REORDERING_CAP_NONE;
-						NvStatusSER = NvAPI_D3D12_GetRaytracingCaps(
-							GetAdapter().GetD3DDevice(),
-							NVAPI_D3D12_RAYTRACING_CAPS_TYPE_THREAD_REORDERING,
-							&ReorderCaps,
-							sizeof(ReorderCaps));
+						{
+							TRACE_CPUPROFILER_EVENT_SCOPE(NvAPI_D3D12_GetRaytracingCaps);
+							NvStatusSER = NvAPI_D3D12_GetRaytracingCaps(
+								D3DDevice,
+								NVAPI_D3D12_RAYTRACING_CAPS_TYPE_THREAD_REORDERING,
+								&ReorderCaps,
+								sizeof(ReorderCaps));
+						}
 
 						if (NvStatusSER != NVAPI_OK || ReorderCaps == NVAPI_D3D12_RAYTRACING_THREAD_REORDERING_CAP_NONE)
 						{
@@ -1693,8 +1699,13 @@ void FD3D12DynamicRHI::Init()
 				{
 					GRHIGlobals.RayTracing.SupportsClusterOps = false;
 
+					ID3D12Device5* Device5 = GetAdapter().GetD3DDevice5();
 					NVAPI_D3D12_RAYTRACING_CLUSTER_OPERATIONS_CAPS ClusterOpsCaps = {};
-					NvAPI_Status NvStatusRayTracingClusterOps = NvAPI_D3D12_GetRaytracingCaps(GetAdapter().GetD3DDevice5(), NVAPI_D3D12_RAYTRACING_CAPS_TYPE_CLUSTER_OPERATIONS, &ClusterOpsCaps, sizeof(NVAPI_D3D12_RAYTRACING_CLUSTER_OPERATIONS_CAPS));
+					NvAPI_Status NvStatusRayTracingClusterOps;
+					{
+						TRACE_CPUPROFILER_EVENT_SCOPE(NvAPI_D3D12_GetRaytracingCaps);
+						NvStatusRayTracingClusterOps = NvAPI_D3D12_GetRaytracingCaps(Device5, NVAPI_D3D12_RAYTRACING_CAPS_TYPE_CLUSTER_OPERATIONS, &ClusterOpsCaps, sizeof(NVAPI_D3D12_RAYTRACING_CLUSTER_OPERATIONS_CAPS));
+					}
 
 					if (NvStatusRayTracingClusterOps == NVAPI_OK && ClusterOpsCaps == NVAPI_D3D12_RAYTRACING_CLUSTER_OPERATIONS_CAP_STANDARD)
 					{
@@ -1724,8 +1735,14 @@ void FD3D12DynamicRHI::Init()
 			GRHIGlobals.RayTracing.SupportsLinearSweptSpheres = false;
 			if (GRHISupportsRayTracing)
 			{
+				ID3D12Device5* Device5 = GetAdapter().GetD3DDevice5();
+				
 				NVAPI_D3D12_RAYTRACING_LINEAR_SWEPT_SPHERES_CAPS LSSCaps = NVAPI_D3D12_RAYTRACING_LINEAR_SWEPT_SPHERES_CAP_NONE;
-				NvAPI_Status NvStatusLSS = NvAPI_D3D12_GetRaytracingCaps(GetAdapter().GetD3DDevice5(), NVAPI_D3D12_RAYTRACING_CAPS_TYPE_LINEAR_SWEPT_SPHERES, &LSSCaps, sizeof(NVAPI_D3D12_RAYTRACING_LINEAR_SWEPT_SPHERES_CAPS));
+				NvAPI_Status NvStatusLSS;
+				{
+					TRACE_CPUPROFILER_EVENT_SCOPE(NvAPI_D3D12_GetRaytracingCaps);
+					NvStatusLSS = NvAPI_D3D12_GetRaytracingCaps(Device5, NVAPI_D3D12_RAYTRACING_CAPS_TYPE_LINEAR_SWEPT_SPHERES, &LSSCaps, sizeof(NVAPI_D3D12_RAYTRACING_LINEAR_SWEPT_SPHERES_CAPS));
+				}
 
 				if (NvStatusLSS == NVAPI_OK && LSSCaps != NVAPI_D3D12_RAYTRACING_LINEAR_SWEPT_SPHERES_CAP_NONE)
 				{
diff --git a/Engine/Source/Runtime/RHI/Private/Windows/WindowsDynamicRHI.cpp b/Engine/Source/Runtime/RHI/Private/Windows/WindowsDynamicRHI.cpp
--- a/Engine/Source/Runtime/RHI/Private/Windows/WindowsDynamicRHI.cpp
+++ b/Engine/Source/Runtime/RHI/Private/Windows/WindowsDynamicRHI.cpp
@@ -1281,6 +1281,8 @@ static void LogTDRSummary()
 
 static IDynamicRHIModule* LoadDynamicRHIModule(ERHIFeatureLevel::Type& DesiredFeatureLevel, const TCHAR*& LoadedRHIModuleName)
 {
+	TRACE_CPUPROFILER_EVENT_SCOPE(LoadDynamicRHIModule);
+
 	// Make sure the DDSPI is initialized before we try and read from it
 	FGenericDataDrivenShaderPlatformInfo::Initialize();
 
diff --git a/Engine/Source/Runtime/RHICore/Private/RHICoreNvidiaAftermath.cpp b/Engine/Source/Runtime/RHICore/Private/RHICoreNvidiaAftermath.cpp
--- a/Engine/Source/Runtime/RHICore/Private/RHICoreNvidiaAftermath.cpp
+++ b/Engine/Source/Runtime/RHICore/Private/RHICoreNvidiaAftermath.cpp
@@ -225,6 +225,8 @@ namespace UE::RHICore::Nvidia::Aftermath
 
 	RHICORE_API void InitializeBeforeDeviceCreation()
 	{
+		TRACE_CPUPROFILER_EVENT_SCOPE(UE::RHICore::Nvidia::Aftermath::InitializeBeforeDeviceCreation);
+
 		LoadAftermathDLL();
 
 		if (!DllHandle)
