Persist D3D12 Adapter Devices after creation until we've choosen adapter.
- Currently we call Release after creation which will discard the device. But we might need the device again later on.
- D3D12CreateDevice can take up to 400ms per adapter.
- Remove SEH around D3D12CreateDevice since D3D12 is guaranteed to be installed

#rnx
#Synced-CL 52153306

diff --git a/Engine/Source/Runtime/D3D12RHI/Private/Windows/WindowsD3D12Device.cpp b/Engine/Source/Runtime/D3D12RHI/Private/Windows/WindowsD3D12Device.cpp
--- a/Engine/Source/Runtime/D3D12RHI/Private/Windows/WindowsD3D12Device.cpp
+++ b/Engine/Source/Runtime/D3D12RHI/Private/Windows/WindowsD3D12Device.cpp
@@ -165,15 +165,6 @@ static void SafeCreateDXGIFactory(IDXGIFactory4** DXGIFactory)
 #endif	//!D3D12_CUSTOM_VIEWPORT_CONSTRUCTOR
 }
 
-/**
- * Returns the minimum D3D feature level required to create based on
- * command line parameters.
- */
-static D3D_FEATURE_LEVEL GetRequiredD3DFeatureLevel()
-{
-	return D3D_FEATURE_LEVEL_11_0;
-}
-
 static D3D_FEATURE_LEVEL FindHighestFeatureLevel(ID3D12Device* Device, D3D_FEATURE_LEVEL MinFeatureLevel)
 {
 	const D3D_FEATURE_LEVEL FeatureLevels[] =
@@ -533,49 +524,45 @@ inline bool GetSupportGPUUploadHeaps(ID3D12Device* device)
 }
 
 /**
- * Attempts to create a D3D12 device for the adapter using at minimum MinFeatureLevel.
+ * Attempts to create a D3D12 device for the adapter or find an existing one.
  * If creation is successful, true is returned.
  */
-static bool SafeTestD3D12CreateDevice(IDXGIAdapter* Adapter, D3D_FEATURE_LEVEL MinFeatureLevel, FD3D12DeviceBasicInfo& OutInfo)
+static bool SafeTestD3D12CreateDevice(IDXGIAdapter* Adapter, LUID AdapterLuid, FD3D12DeviceBasicInfo& OutInfo)
 {
-	__try
+	constexpr D3D_FEATURE_LEVEL MinFeatureLevel = D3D_FEATURE_LEVEL_11_0;
+
+	uint64 Luid = BitCast<uint64>(AdapterLuid);
+	TRefCountPtr<ID3D12Device> Device = FWindowsD3D::FindCachedD3D12AdapterDevice(Luid);
+	if (!Device.IsValid())
 	{
-		ID3D12Device* Device = nullptr;
-		const HRESULT D3D12CreateDeviceResult = D3D12CreateDevice(Adapter, MinFeatureLevel, IID_PPV_ARGS(&Device));
-		if (SUCCEEDED(D3D12CreateDeviceResult))
+		const HRESULT D3D12CreateDeviceResult = D3D12CreateDevice(Adapter, MinFeatureLevel, IID_PPV_ARGS(Device.GetInitReference()));
+		if (!SUCCEEDED(D3D12CreateDeviceResult))
 		{
-			OutInfo.MaxFeatureLevel = FindHighestFeatureLevel(Device, MinFeatureLevel);
-			OutInfo.MaxShaderModel = FindHighestShaderModel(Device);
-			GetResourceTiers(Device, OutInfo.ResourceBindingTier, OutInfo.ResourceHeapTier);
-			OutInfo.NumDeviceNodes = Device->GetNodeCount();
+			UE_LOG(LogD3D12RHI, Log, TEXT("D3D12CreateDevice failed with code 0x%08X"), static_cast<int32>(D3D12CreateDeviceResult));
+			return false;
+		}
 
-			OutInfo.bSupportsWaveOps = GetSupportsWaveOps(Device);
-			OutInfo.bSupportsAtomic64 = GetSupportsAtomic64(Adapter, Device);
+		FWindowsD3D::CacheD3D12AdapterDevice(Luid, Device);
+	}
 
-			OutInfo.MaxRHIFeatureLevel = FindMaxRHIFeatureLevel(OutInfo.MaxFeatureLevel, OutInfo.MaxShaderModel, OutInfo.ResourceBindingTier, OutInfo.bSupportsWaveOps, OutInfo.bSupportsAtomic64);
+	OutInfo.MaxFeatureLevel = FindHighestFeatureLevel(Device, MinFeatureLevel);
+	OutInfo.MaxShaderModel = FindHighestShaderModel(Device);
+	GetResourceTiers(Device, OutInfo.ResourceBindingTier, OutInfo.ResourceHeapTier);
+	OutInfo.NumDeviceNodes = Device->GetNodeCount();
 
-			if (!GetDeviceArchitectureInformation(Device, OutInfo.bTileBasedRenderer, OutInfo.bUMA, OutInfo.bCacheCoherentUMA))
-			{
-				OutInfo.bTileBasedRenderer = OutInfo.bUMA = OutInfo.bCacheCoherentUMA = false;
-			}
+	OutInfo.bSupportsWaveOps = GetSupportsWaveOps(Device);
+	OutInfo.bSupportsAtomic64 = GetSupportsAtomic64(Adapter, Device);
 
-			OutInfo.bSupportsGPUUploadHeaps = GetSupportGPUUploadHeaps(Device);
+	OutInfo.MaxRHIFeatureLevel = FindMaxRHIFeatureLevel(OutInfo.MaxFeatureLevel, OutInfo.MaxShaderModel, OutInfo.ResourceBindingTier, OutInfo.bSupportsWaveOps, OutInfo.bSupportsAtomic64);
 
-			Device->Release();
-			return true;
-		}
-		else
-		{
-			UE_LOG(LogD3D12RHI, Log, TEXT("D3D12CreateDevice failed with code 0x%08X"), static_cast<int32>(D3D12CreateDeviceResult));
-		}
-	}
-	__except (IsDelayLoadException(GetExceptionInformation()))
+	if (!GetDeviceArchitectureInformation(Device, OutInfo.bTileBasedRenderer, OutInfo.bUMA, OutInfo.bCacheCoherentUMA))
 	{
-		// We suppress warning C6322: Empty _except block. Appropriate checks are made upon returning. 
-		CA_SUPPRESS(6322);
+		OutInfo.bTileBasedRenderer = OutInfo.bUMA = OutInfo.bCacheCoherentUMA = false;
 	}
 
-	return false;
+	OutInfo.bSupportsGPUUploadHeaps = GetSupportGPUUploadHeaps(Device);
+
+	return true;
 }
 
 static bool SupportsDepthBoundsTest(FD3D12DynamicRHI* D3DRHI)
@@ -898,7 +885,6 @@ void FD3D12DynamicRHIModule::FindAdapter()
 	const uint64 AdapterLuid = FWindowsD3D::ChooseD3D12Adapter(HmdGraphicsAdapterLuid);
 
 	TRefCountPtr<IDXGIAdapter> TempAdapter;
-	const D3D_FEATURE_LEVEL MinRequiredFeatureLevel = GetRequiredD3DFeatureLevel();
 
 	int GpuPreferenceInt = DXGI_GPU_PREFERENCE_HIGH_PERFORMANCE;
 	FParse::Value(FCommandLine::Get(), TEXT("-gpupreference="), GpuPreferenceInt);
@@ -930,7 +916,7 @@ void FD3D12DynamicRHIModule::FindAdapter()
 			}
 #endif
 
-			if (SafeTestD3D12CreateDevice(TempAdapter, MinRequiredFeatureLevel, DeviceInfo))
+			if (SafeTestD3D12CreateDevice(TempAdapter, AdapterDesc.AdapterLuid, DeviceInfo))
 			{
 				check(DeviceInfo.NumDeviceNodes > 0);
 
diff --git a/Engine/Source/Runtime/RHI/Private/Windows/WindowsDynamicRHI.cpp b/Engine/Source/Runtime/RHI/Private/Windows/WindowsDynamicRHI.cpp
--- a/Engine/Source/Runtime/RHI/Private/Windows/WindowsDynamicRHI.cpp
+++ b/Engine/Source/Runtime/RHI/Private/Windows/WindowsDynamicRHI.cpp
@@ -1392,6 +1392,8 @@ FDynamicRHI* PlatformCreateDynamicRHI()
 	const TCHAR* LoadedRHIModuleName;
 	IDynamicRHIModule* DynamicRHIModule = LoadDynamicRHIModule(RequestedFeatureLevel, LoadedRHIModuleName);
 
+	FWindowsD3D::ReleaseCachedD3D12AdapterDevices();
+
 	if (DynamicRHIModule)
 	{
 		// Run ConfigRules to allow hotfixing RHI. Things like GRHIDeviceId, GRHIVendorId are only available at this stage
diff --git a/Engine/Source/Runtime/Windows/WindowsD3D/Internal/Windows/WindowsD3D.h b/Engine/Source/Runtime/Windows/WindowsD3D/Internal/Windows/WindowsD3D.h
--- a/Engine/Source/Runtime/Windows/WindowsD3D/Internal/Windows/WindowsD3D.h
+++ b/Engine/Source/Runtime/Windows/WindowsD3D/Internal/Windows/WindowsD3D.h
@@ -4,6 +4,9 @@
 
 #include "Containers/UnrealString.h"
 #include "HAL/Platform.h"
+#include "Templates/RefCounting.h"
+
+struct ID3D12Device;
 
 struct FWindowsGPUInfo
 {
@@ -20,4 +23,8 @@ public:
 	static WINDOWSD3D_API FWindowsGPUInfo GetGPUInfoByDedicatedMemory();
 	static WINDOWSD3D_API uint64 ChooseD3D11Adapter(uint64 DesiredAdapterLuid);
 	static WINDOWSD3D_API uint64 ChooseD3D12Adapter(uint64 DesiredAdapterLuid);
+
+	static WINDOWSD3D_API TRefCountPtr<ID3D12Device> FindCachedD3D12AdapterDevice(uint64 AdapterLuid);
+	static WINDOWSD3D_API void CacheD3D12AdapterDevice(uint64 AdapterLuid, TRefCountPtr<ID3D12Device> Device);
+	static WINDOWSD3D_API void ReleaseCachedD3D12AdapterDevices();
 };
diff --git a/Engine/Source/Runtime/Windows/WindowsD3D/Private/WindowsD3D.cpp b/Engine/Source/Runtime/Windows/WindowsD3D/Private/WindowsD3D.cpp
--- a/Engine/Source/Runtime/Windows/WindowsD3D/Private/WindowsD3D.cpp
+++ b/Engine/Source/Runtime/Windows/WindowsD3D/Private/WindowsD3D.cpp
@@ -18,6 +18,8 @@ THIRD_PARTY_INCLUDES_END
 
 DEFINE_LOG_CATEGORY_STATIC(LogD3D, Log, All);
 
+static TMap<uint64, TRefCountPtr<ID3D12Device>> D3D12AdapterDevices;
+
 enum class ED3DGpuVendorId : uint32
 {
 	Unknown = 0xffffffff,
@@ -199,37 +201,32 @@ static bool IsDelayLoadException(PEXCEPTION_POINTERS ExceptionPointers)
 }
 
 /**
- * Attempts to create a D3D12 device for the adapter using at minimum MinFeatureLevel.
+ * Attempts to create a D3D12 device for the adapter or find an existing one.
  * If creation is successful, true is returned.
  */
-static bool SafeTestD3D12CreateDevice(IDXGIAdapter* Adapter, D3D_FEATURE_LEVEL MinFeatureLevel, bool& bOutUMA)
+static bool SafeTestD3D12CreateDevice(IDXGIAdapter* Adapter, LUID AdapterLuid, bool& bOutUMA)
 {
 	bOutUMA = false;
 
-	__try
-	{
-		ID3D12Device* Device = nullptr;
-		const HRESULT D3D12CreateDeviceResult = D3D12CreateDevice(Adapter, MinFeatureLevel, IID_PPV_ARGS(&Device));
-		if (SUCCEEDED(D3D12CreateDeviceResult))
-		{
-			D3D12_FEATURE_DATA_ARCHITECTURE data{};
-			bOutUMA = SUCCEEDED(Device->CheckFeatureSupport(D3D12_FEATURE_ARCHITECTURE, &data, sizeof(D3D12_FEATURE_DATA_ARCHITECTURE))) && data.UMA;
+	constexpr D3D_FEATURE_LEVEL MinFeatureLevel = D3D_FEATURE_LEVEL_11_0;
 
-			Device->Release();
-			return true;
-		}
-		else
+	uint64 Luid = BitCast<uint64>(AdapterLuid);
+	TRefCountPtr<ID3D12Device> Device = FWindowsD3D::FindCachedD3D12AdapterDevice(Luid);
+	if (!Device)
+	{
+		const HRESULT D3D12CreateDeviceResult = D3D12CreateDevice(Adapter, MinFeatureLevel, IID_PPV_ARGS(Device.GetInitReference()));
+		if (!SUCCEEDED(D3D12CreateDeviceResult))
 		{
 			UE_LOG(LogD3D, Log, TEXT("D3D12CreateDevice failed with code 0x%08X"), static_cast<int32>(D3D12CreateDeviceResult));
+			return false;
 		}
-	}
-	__except (IsDelayLoadException(GetExceptionInformation()))
-	{
-		// We suppress warning C6322: Empty _except block. Appropriate checks are made upon returning. 
-		CA_SUPPRESS(6322);
+
+		FWindowsD3D::CacheD3D12AdapterDevice(Luid, Device);
 	}
 
-	return false;
+	D3D12_FEATURE_DATA_ARCHITECTURE data{};
+	bOutUMA = SUCCEEDED(Device->CheckFeatureSupport(D3D12_FEATURE_ARCHITECTURE, &data, sizeof(D3D12_FEATURE_DATA_ARCHITECTURE))) && data.UMA;
+	return true;
 }
 
 /**
@@ -529,7 +526,6 @@ uint64 FWindowsD3D::ChooseD3D12Adapter(uint64 DesiredAdapterLuid)
 	const bool bFavorDiscreteAdapter = CVarExplicitAdapterValue == -1;
 
 	TRefCountPtr<IDXGIAdapter> TempAdapter;
-	const D3D_FEATURE_LEVEL MinRequiredFeatureLevel = D3D_FEATURE_LEVEL_11_0;
 
 	TOptional<LUID> PreferredAdapterLuid;
 	TOptional<LUID> FirstDiscreteAdapterLuid;
@@ -560,7 +556,7 @@ uint64 FWindowsD3D::ChooseD3D12Adapter(uint64 DesiredAdapterLuid)
 		{
 			bool bUMA;
 			DXGI_ADAPTER_DESC AdapterDesc{};
-			if (SUCCEEDED(TempAdapter->GetDesc(&AdapterDesc)) && SafeTestD3D12CreateDevice(TempAdapter, MinRequiredFeatureLevel, bUMA))
+			if (SUCCEEDED(TempAdapter->GetDesc(&AdapterDesc)) && SafeTestD3D12CreateDevice(TempAdapter, AdapterDesc.AdapterLuid, bUMA))
 			{
 				const bool bIsWARP = ConvertToGpuVendorId(AdapterDesc.VendorId) == ED3DGpuVendorId::Microsoft;
 
@@ -650,3 +646,24 @@ uint64 FWindowsD3D::ChooseD3D12Adapter(uint64 DesiredAdapterLuid)
 	}
 	return Luid;
 }
+
+TRefCountPtr<ID3D12Device> FWindowsD3D::FindCachedD3D12AdapterDevice(uint64 AdapterLuid)
+{
+	TRefCountPtr<ID3D12Device>* DevicePtr = D3D12AdapterDevices.Find(AdapterLuid);
+	if (DevicePtr)
+	{
+		return *DevicePtr;
+	}
+
+	return nullptr;
+}
+
+void FWindowsD3D::CacheD3D12AdapterDevice(uint64 AdapterLuid, TRefCountPtr<ID3D12Device> Device)
+{
+	D3D12AdapterDevices.Emplace(MoveTemp(AdapterLuid), MoveTemp(Device));
+}
+
+void FWindowsD3D::ReleaseCachedD3D12AdapterDevices()
+{
+	D3D12AdapterDevices.Reset();
+}
