Adjust async compute usage for Intel GPUs on D3D12

#Synced-CL 52232959

diff --git a/Engine/Source/Runtime/D3D12RHI/Private/D3D12Adapter.cpp b/Engine/Source/Runtime/D3D12RHI/Private/D3D12Adapter.cpp
--- a/Engine/Source/Runtime/D3D12RHI/Private/D3D12Adapter.cpp
+++ b/Engine/Source/Runtime/D3D12RHI/Private/D3D12Adapter.cpp
@@ -1417,19 +1417,23 @@ void FD3D12Adapter::InitializeDevices()
 			{
 				GSupportsEfficientAsyncCompute = true;
 			}
-			else if (IsRHIDeviceIntel() && GMaxRHIFeatureLevel >= ERHIFeatureLevel::SM6)
+			else if (IsRHIDeviceIntel())
 			{
-				GSupportsEfficientAsyncCompute = true;
+				// Unknown/Arc
+				if ((GRHIGlobals.GpuInfo.AdapterGeneration == 0) || (GRHIGlobals.GpuInfo.AdapterGeneration >= 100))
+				{
+					GSupportsEfficientAsyncCompute = true;
+				}
 			}
 #if PLATFORM_WINDOWS
-			else
+			else if (IsRHIDeviceNVIDIA())
 			{
 				D3D12_FEATURE_DATA_D3D12_OPTIONS6 D3D12Caps6{};
 				HRESULT Options6HR = RootDevice->CheckFeatureSupport(D3D12_FEATURE_D3D12_OPTIONS6, &D3D12Caps6, sizeof(D3D12Caps6));
 
-				// Allow async compute by default on nVidia cards which support PerPrimitiveShadingRateSupportedWithViewportIndexing 
-				// this should be a good metric according to nVidia itself (this is set for Ampere and newer cards)
-				if (IsRHIDeviceNVIDIA() && Options6HR == S_OK && D3D12Caps6.PerPrimitiveShadingRateSupportedWithViewportIndexing)
+				// Allow async compute by default on NVIDIA cards which support PerPrimitiveShadingRateSupportedWithViewportIndexing
+				// this should be a good metric according to NVIDIA itself (this is set for Ampere and newer cards)
+				if (Options6HR == S_OK && D3D12Caps6.PerPrimitiveShadingRateSupportedWithViewportIndexing)
 				{
 					GSupportsEfficientAsyncCompute = true;
 				}
