Adding ResourceCollection support to a few binding paths that were missing them.

#Synced-CL 52308469

diff --git a/Engine/Source/Runtime/Apple/MetalRHI/Private/MetalRayTracing.cpp b/Engine/Source/Runtime/Apple/MetalRHI/Private/MetalRayTracing.cpp
--- a/Engine/Source/Runtime/Apple/MetalRHI/Private/MetalRayTracing.cpp
+++ b/Engine/Source/Runtime/Apple/MetalRHI/Private/MetalRayTracing.cpp
@@ -12,6 +12,7 @@
 #include "MetalRHIContext.h"
 #include "MetalShaderTypes.h"
 #include "MetalBindlessDescriptors.h"
+#include "MetalResourceCollection.h"
 #include "MetalStaticSamplers.h"
 #include "Shaders/Types/MetalRayShader.h"
 #include "BuiltInRayTracingShaders.h"
@@ -2791,7 +2792,21 @@ static void DispatchRays(FMetalRHICommandContext& Context,
 					FMetalSurface* Surface = GetMetalSurfaceFromRHITexture(Texture);
 					StateCache.IRMakeTextureResident(EMetalShaderStages::Compute, Surface->Texture.get());
 				}
+				break;
 			}
+#if PLATFORM_SUPPORTS_BINDLESS_RENDERING
+			case FRHIShaderParameterResource::EType::ResourceCollection:
+			{
+				if (FMetalResourceCollection* MetalResourceCollection = ResourceCast(static_cast<FRHIResourceCollection*>(Parameter.Resource)))
+				{
+					if (FRHIShaderResourceView* SRV = MetalResourceCollection->GetShaderResourceView())
+					{
+						StateCache.IRMakeSRVResident(EMetalShaderStages::Compute, static_cast<FMetalShaderResourceView*>(SRV));
+					}
+				}
+				break;
+			}
+#endif
 			default:
 			{
 				break;
diff --git a/Engine/Source/Runtime/D3D12RHI/Private/D3D12RayTracing.cpp b/Engine/Source/Runtime/D3D12RHI/Private/D3D12RayTracing.cpp
--- a/Engine/Source/Runtime/D3D12RHI/Private/D3D12RayTracing.cpp
+++ b/Engine/Source/Runtime/D3D12RHI/Private/D3D12RayTracing.cpp
@@ -5532,6 +5532,26 @@ static bool SetRayTracingShaderResources(
 			BoundSRVMask |= 1ull << Index;
 
 			Binder.AddReferencedShaderResource(SRV->GetShaderResource());
+
+			// Ensure all member SRVs referenced by the collection buffer are resident.
+			for (FD3D12ShaderResourceView* MemberSRV : D3D12ResourceCollection->AllSrvs)
+			{
+				Binder.AddReferencedShaderResource(MemberSRV->GetShaderResource());
+
+				if (FD3D12RayTracingScene* ReferencedRayTracingScene = MemberSRV->GetRayTracingScene())
+				{
+					ReferencedRayTracingScenes.Add(ReferencedRayTracingScene);
+				}
+			}
+
+			// Resolve texture references to their current textures and ensure residency.
+			for (FD3D12RHITextureReference* TextureReference : D3D12ResourceCollection->AllTextureReferences)
+			{
+				if (TextureReference)
+				{
+					Binder.AddReferencedTexture(TextureReference);
+				}
+			}
 		}
 
 		void SetSampler(FRHISamplerState* RHISampler, uint8 Index)
@@ -5572,6 +5592,9 @@ static bool SetRayTracingShaderResources(
 				case FRHIShaderParameterResource::EType::Sampler:
 					Bindings.SetSampler(static_cast<FRHISamplerState*>(Resource), BindlessParameterIndex);
 					break;
+				case FRHIShaderParameterResource::EType::ResourceCollection:
+					Bindings.SetResourceCollection(static_cast<FRHIResourceCollection*>(Resource), BindlessParameterIndex);
+					break;
 			}
 		}
 	}
diff --git a/Engine/Source/Runtime/D3D12RHI/Private/D3D12WorkGraph.cpp b/Engine/Source/Runtime/D3D12RHI/Private/D3D12WorkGraph.cpp
--- a/Engine/Source/Runtime/D3D12RHI/Private/D3D12WorkGraph.cpp
+++ b/Engine/Source/Runtime/D3D12RHI/Private/D3D12WorkGraph.cpp
@@ -472,9 +472,19 @@ struct FWorkGraphShaderBundleBinder
 		FD3D12ShaderResourceView* SRV = D3D12ResourceCollection ? D3D12ResourceCollection->GetShaderResourceView() : nullptr;
 
 		check(bBindlessResources);
-		if (bBindlessResources)
+		if (SRV && bBindlessResources)
 		{
 			Context.StateCache.QueueBindlessSRV(Frequency, SRV);
+			Context.StateCache.QueueBindlessSRVs(Frequency, D3D12ResourceCollection->AllSrvs);
+
+			// Resolve texture references to their current textures and ensure residency.
+			for (FD3D12RHITextureReference* TextureReference : D3D12ResourceCollection->AllTextureReferences)
+			{
+				if (FD3D12Texture* Texture = GetD3D12TextureFromRHITexture(TextureReference))
+				{
+					Context.StateCache.QueueBindlessSRV(Frequency, Texture->GetShaderResourceView());
+				}
+			}
 		}
 	}
 };
