Add GD3D12ResidencyDebugBudgetMB as CVar for stress testing residency manager

#Synced-CL 52149918

diff --git a/Engine/Source/Runtime/D3D12RHI/Private/D3D12Adapter.cpp b/Engine/Source/Runtime/D3D12RHI/Private/D3D12Adapter.cpp
--- a/Engine/Source/Runtime/D3D12RHI/Private/D3D12Adapter.cpp
+++ b/Engine/Source/Runtime/D3D12RHI/Private/D3D12Adapter.cpp
@@ -201,6 +201,16 @@ static FAutoConsoleVariableRef CVarD3D12EvictAllResidentResourcesInBackground(
 	TEXT("Force D3D12 resource residency manager to evict all tracked unused resources when the application is not focused\n"),
 	ECVF_Default);
 
+int32 GD3D12ResidencyDebugBudgetMB = -1;
+static FAutoConsoleVariableRef CVarD3D12ResidencyDebugBudgetMB(
+	TEXT("D3D12.ResidencyDebugBudgetMB"),
+	GD3D12ResidencyDebugBudgetMB,
+	TEXT("When >= 0, overrides the local memory budget for the residency manager (in MB).\n")
+	TEXT("Use a small value (e.g. 128) to force aggressive eviction/MakeResident churn,\n")
+	TEXT("exposing missing UpdateResidency calls that would otherwise be hidden by implicit driver paging.\n")
+	TEXT("Set to 0 for maximum eviction pressure. Default: -1 (disabled)."),
+	ECVF_Default);
+
 #endif // PLATFORM_WINDOWS
 
 #if D3D12_SUPPORTS_INFO_QUEUE
@@ -1881,7 +1891,14 @@ const FRHIMemoryStats& FD3D12Adapter::CollectMemoryStats()
 	// performance degradation when paged-out resources are used for rendering before VidMm pages them back in (which can take a long time).
 	// By overriding the budget to 0 and stopping rendering at the high-level, we can immediately free VRAM and avoid VidMm paging.
 	const bool bEvictResidentResources = GEnableResidencyManagement && GD3D12EvictAllResidentResourcesInBackground && !FApp::HasFocus();
-	const uint64 LocalMemoryBudgetLimit = bEvictResidentResources ? 0 : MemoryStats.BudgetLocal;
+	uint64 LocalMemoryBudgetLimit = bEvictResidentResources ? 0 : MemoryStats.BudgetLocal;
+
+	// Debug override: force a specific budget to stress-test residency management.
+	if (GD3D12ResidencyDebugBudgetMB >= 0)
+	{
+		LocalMemoryBudgetLimit = (uint64)GD3D12ResidencyDebugBudgetMB * 1024 * 1024;
+	}
+
 	for (uint32 GPUIndex : FRHIGPUMask::All())
 	{
 		GetDevice(GPUIndex)->GetResidencyManager().SetLocalMemoryBudgetLimit(LocalMemoryBudgetLimit);
