Removing bAllowConfidentialPlatformDefines from D3D12RHI.Build.cs

#Synced-CL 52353451

diff --git a/Engine/Source/Runtime/D3D12RHI/D3D12RHI.Build.cs b/Engine/Source/Runtime/D3D12RHI/D3D12RHI.Build.cs
--- a/Engine/Source/Runtime/D3D12RHI/D3D12RHI.Build.cs
+++ b/Engine/Source/Runtime/D3D12RHI/D3D12RHI.Build.cs
@@ -8,8 +8,6 @@ public class D3D12RHI : ModuleRules
 {
 	public D3D12RHI(ReadOnlyTargetRules Target) : base(Target)
 	{
-		bAllowConfidentialPlatformDefines = true;
-
 		PrivateIncludePathModuleNames.AddAll(
 			"Shaders"
 		);
