Clear the vertex buffer cache when a render pass ends.

This prevents validation errors when RDG transitions buffers to other states after it's done using them as stream sources. If we don't clear the cache, it's possible that a buffer is transitioned to a state which doesn't include D3D12_RESOURCE_STATE_VERTEX_AND_CONSTANT_BUFFER, but it's still in the cache because the next draw doesn't use the slot it was bound to. The shader doesn't use the buffer, but the validation layer can't know that, so it produces an error when the cache binds the unused buffer in the wrong state.

#rnx

#Synced-CL 52466084

diff --git a/Engine/Source/Runtime/D3D12RHI/Private/D3D12CommandContext.h b/Engine/Source/Runtime/D3D12RHI/Private/D3D12CommandContext.h
--- a/Engine/Source/Runtime/D3D12RHI/Private/D3D12CommandContext.h
+++ b/Engine/Source/Runtime/D3D12RHI/Private/D3D12CommandContext.h
@@ -736,6 +736,12 @@ public:
 		{
 			ResolveTexture(Info);
 		});
+
+		// Buffers used as vertex streams in this pass can be transitioned to other states after the pass ends. If such a buffer
+		// was bound to a slot that's not used in a subsequent draw, the cache will apply the stale binding, resulting in validation
+		// layer errors if the new state doesn't include D3D12_RESOURCE_STATE_VERTEX_AND_CONSTANT_BUFFER (despite the shader not actually
+		// using the buffer - the validation layer can't know that). We clear the VB cache to prevent this case.
+		StateCache.ClearVertexBuffers();
 	}
 
 	void ResolveTexture(UE::RHICore::FResolveTextureInfo Info);
diff --git a/Engine/Source/Runtime/D3D12RHI/Private/D3D12StateCachePrivate.h b/Engine/Source/Runtime/D3D12RHI/Private/D3D12StateCachePrivate.h
--- a/Engine/Source/Runtime/D3D12RHI/Private/D3D12StateCachePrivate.h
+++ b/Engine/Source/Runtime/D3D12RHI/Private/D3D12StateCachePrivate.h
@@ -634,6 +634,12 @@ public:
 		}
 	}
 
+	D3D12_STATE_CACHE_INLINE void ClearVertexBuffers()
+	{
+		PipelineState.Graphics.VBCache.Clear();
+		bNeedSetVB = true;
+	}
+
 public:
 
 	D3D12_STATE_CACHE_INLINE void SetIndexBuffer(const FD3D12ResourceLocation& IndexBufferLocation, DXGI_FORMAT Format, uint32 Offset)
