CopyTextureRegionChecked assert improvements -- both source and destination name are printed.

#rnx

#Synced-CL 52530005

diff --git a/Engine/Source/Runtime/D3D12RHI/Private/D3D12CommandContext.cpp b/Engine/Source/Runtime/D3D12RHI/Private/D3D12CommandContext.cpp
--- a/Engine/Source/Runtime/D3D12RHI/Private/D3D12CommandContext.cpp
+++ b/Engine/Source/Runtime/D3D12RHI/Private/D3D12CommandContext.cpp
@@ -1006,7 +1006,7 @@ void FD3D12CommandContextRedirector::RHITransferResources(TConstArrayView<FTrans
 				CD3DX12_TEXTURE_COPY_LOCATION SrcLocation(SrcTexture->GetResource()->GetResource(), 0);
 				CD3DX12_TEXTURE_COPY_LOCATION DstLocation(DstTexture->GetResource()->GetResource(), 0);
 
-				CopyContext->CopyTextureRegionChecked(&DstLocation, Box.left, Box.top, Box.front, DstTexture->GetFormat(), &SrcLocation, &Box, SrcTexture->GetFormat(), DstTexture->GetName());
+				CopyContext->CopyTextureRegionChecked(&DstLocation, Box.left, Box.top, Box.front, DstTexture->GetFormat(), &SrcLocation, &Box, SrcTexture->GetFormat(), DstTexture->GetName(), SrcTexture->GetName());
 
 				CopyContext->UpdateResidency(SrcTexture->GetResource());
 				CopyContext->UpdateResidency(DstTexture->GetResource());
@@ -1169,7 +1169,7 @@ void FD3D12CommandContextRedirector::RHICrossGPUTransfer(TConstArrayView<FTransf
 				CD3DX12_TEXTURE_COPY_LOCATION SrcLocation(SrcTexture->GetResource()->GetResource(), 0);
 				CD3DX12_TEXTURE_COPY_LOCATION DstLocation(DstTexture->GetResource()->GetResource(), 0);
 
-				SrcContext->CopyTextureRegionChecked(&DstLocation, Box.left, Box.top, Box.front, DstTexture->GetFormat(), &SrcLocation, &Box, SrcTexture->GetFormat(), DstTexture->GetName());
+				SrcContext->CopyTextureRegionChecked(&DstLocation, Box.left, Box.top, Box.front, DstTexture->GetFormat(), &SrcLocation, &Box, SrcTexture->GetFormat(), DstTexture->GetName(), SrcTexture->GetName());
 				
 				SrcContext->UpdateResidency(SrcTexture->GetResource());
 				SrcContext->UpdateResidency(DstTexture->GetResource());
diff --git a/Engine/Source/Runtime/D3D12RHI/Private/D3D12CommandContext.h b/Engine/Source/Runtime/D3D12RHI/Private/D3D12CommandContext.h
--- a/Engine/Source/Runtime/D3D12RHI/Private/D3D12CommandContext.h
+++ b/Engine/Source/Runtime/D3D12RHI/Private/D3D12CommandContext.h
@@ -857,7 +857,7 @@ public:
 	FORCENOINLINE void CopyTextureRegionChecked(
 		const D3D12_TEXTURE_COPY_LOCATION* DestCopyLocation, int DestX, int DestY, int DestZ, EPixelFormat DestPixelFormat,
 		const D3D12_TEXTURE_COPY_LOCATION* SourceCopyLocation, const D3D12_BOX* SourceBox, EPixelFormat SourcePixelFormat,
-		const FName& DebugName
+		const FName& DestName, const FName& SourceName
 	);
 
 protected:
diff --git a/Engine/Source/Runtime/D3D12RHI/Private/D3D12RenderTarget.cpp b/Engine/Source/Runtime/D3D12RHI/Private/D3D12RenderTarget.cpp
--- a/Engine/Source/Runtime/D3D12RHI/Private/D3D12RenderTarget.cpp
+++ b/Engine/Source/Runtime/D3D12RHI/Private/D3D12RenderTarget.cpp
@@ -476,7 +476,8 @@ TRefCountPtr<FD3D12Resource> FD3D12DynamicRHI::GetStagingTexture(FRHITexture* Te
 				&SourceCopyLocation,
 				RectPtr,
 				Texture->GetFormat(),
-				Texture->GetName()
+				TempTexture2D->GetName(),		// DestName
+				Texture->GetName()				// SourceName
 			);
 
 			Context.UpdateResidency(Texture->GetResource());
@@ -810,7 +811,8 @@ void FD3D12DynamicRHI::ReadSurfaceDataMSAARaw(FRHITexture* TextureRHI, FIntRect
 				&SourceCopyLocation,
 				&Rect,
 				Texture->GetFormat(),
-				Texture->GetName()
+				StagingTexture2D->GetName(),		// DestName
+				Texture->GetName()					// SourceName
 			);
 
 			Context.UpdateResidency(StagingTexture2D);
@@ -994,7 +996,8 @@ void FD3D12DynamicRHI::RHIReadSurfaceFloatData(FRHITexture* TextureRHI, FIntRect
 			&SourceCopyLocation,
 			&Rect,
 			Texture->GetFormat(),
-			Texture->GetName()
+			TempTexture2D->GetName(),		// DestName
+			Texture->GetName()				// SourceName
 		);
 
 		Context.UpdateResidency(Texture->GetResource());
@@ -1115,7 +1118,8 @@ void FD3D12DynamicRHI::RHIRead3DSurfaceFloatData(FRHITexture* TextureRHI, FIntRe
 			&SourceCopyLocation,
 			&Rect,
 			Texture->GetFormat(),
-			Texture->GetName()
+			TempTexture3D->GetName(),		// DestName
+			Texture->GetName()				// SourceName
 		);
 
 		Context.UpdateResidency(Texture->GetResource());
diff --git a/Engine/Source/Runtime/D3D12RHI/Private/D3D12Texture.cpp b/Engine/Source/Runtime/D3D12RHI/Private/D3D12Texture.cpp
--- a/Engine/Source/Runtime/D3D12RHI/Private/D3D12Texture.cpp
+++ b/Engine/Source/Runtime/D3D12RHI/Private/D3D12Texture.cpp
@@ -1520,7 +1520,7 @@ FString GetTextureCopyRegionString(const FSubresourceIndices& Indices, const D3D
 FORCENOINLINE void FD3D12CommandContext::CopyTextureRegionChecked(
 	const D3D12_TEXTURE_COPY_LOCATION* DestCopyLocation, int DestX, int DestY, int DestZ, EPixelFormat DestPixelFormat,
 	const D3D12_TEXTURE_COPY_LOCATION* SourceCopyLocation, const D3D12_BOX* SourceBox, EPixelFormat SourcePixelFormat,
-	const FName& DebugName
+	const FName& DestName, const FName& SourceName
 )
 {
 #if ENABLE_COPY_TEXTURE_REGION_CHECK
@@ -1545,20 +1545,20 @@ FORCENOINLINE void FD3D12CommandContext::CopyTextureRegionChecked(
 			const uint32 SubresourceHeight = AlignArbitrary<uint32>(FMath::Max<uint32>(1u, SourceDesc.Height >> Indices.MipIndex), SourcePixelFormatInfo.BlockSizeY);
 			const uint32 SubresourceDepth  = AlignArbitrary<uint32>(SourceDesc.Dimension == D3D12_RESOURCE_DIMENSION_TEXTURE3D ? FMath::Max<uint32>(1u, SourceDesc.DepthOrArraySize >> Indices.MipIndex) : 1, SourcePixelFormatInfo.BlockSizeZ);
 
-			RHI_BREADCRUMB_CHECK_SHIPPINGF(*this, (CopyBox.left % SourcePixelFormatInfo.BlockSizeX) == 0 && (CopyBox.right % SourcePixelFormatInfo.BlockSizeX) == 0, TEXT("Width not aligned to block size for: '%s', %s"), *DebugName.ToString(), LAZY_COPY_DECORATOR);
-			RHI_BREADCRUMB_CHECK_SHIPPINGF(*this, (CopyBox.top % SourcePixelFormatInfo.BlockSizeY) == 0 && (CopyBox.bottom % SourcePixelFormatInfo.BlockSizeY) == 0, TEXT("Height not aligned to block size for: '%s', %s"), *DebugName.ToString(), LAZY_COPY_DECORATOR);
+			RHI_BREADCRUMB_CHECK_SHIPPINGF(*this, (CopyBox.left % SourcePixelFormatInfo.BlockSizeX) == 0 && (CopyBox.right % SourcePixelFormatInfo.BlockSizeX) == 0, TEXT("Width not aligned to block size for: '%s' -> '%s', %s"), *SourceName.ToString(), *DestName.ToString(), LAZY_COPY_DECORATOR);
+			RHI_BREADCRUMB_CHECK_SHIPPINGF(*this, (CopyBox.top % SourcePixelFormatInfo.BlockSizeY) == 0 && (CopyBox.bottom % SourcePixelFormatInfo.BlockSizeY) == 0, TEXT("Height not aligned to block size for: '%s' -> '%s', %s"), *SourceName.ToString(), *DestName.ToString(), LAZY_COPY_DECORATOR);
 				
-			RHI_BREADCRUMB_CHECK_SHIPPINGF(*this, CopyBox.left <= CopyBox.right && CopyBox.right <= SubresourceWidth, TEXT("Source width out of bounds for: '%s', %s"), *DebugName.ToString(), LAZY_COPY_DECORATOR);
-			RHI_BREADCRUMB_CHECK_SHIPPINGF(*this, CopyBox.top <= CopyBox.bottom && CopyBox.bottom <= SubresourceHeight, TEXT("Source height out of bounds for: '%s', %s"), *DebugName.ToString(), LAZY_COPY_DECORATOR);
-			RHI_BREADCRUMB_CHECK_SHIPPINGF(*this, CopyBox.front <= CopyBox.back && CopyBox.back <= SubresourceDepth, TEXT("Source depth out of bounds for: '%s', %s"), *DebugName.ToString(), LAZY_COPY_DECORATOR);
+			RHI_BREADCRUMB_CHECK_SHIPPINGF(*this, CopyBox.left <= CopyBox.right && CopyBox.right <= SubresourceWidth, TEXT("Source width out of bounds for: '%s' -> '%s', %s"), *SourceName.ToString(), *DestName.ToString(), LAZY_COPY_DECORATOR);
+			RHI_BREADCRUMB_CHECK_SHIPPINGF(*this, CopyBox.top <= CopyBox.bottom && CopyBox.bottom <= SubresourceHeight, TEXT("Source height out of bounds for: '%s' -> '%s', %s"), *SourceName.ToString(), *DestName.ToString(), LAZY_COPY_DECORATOR);
+			RHI_BREADCRUMB_CHECK_SHIPPINGF(*this, CopyBox.front <= CopyBox.back && CopyBox.back <= SubresourceDepth, TEXT("Source depth out of bounds for:  '%s' -> '%s', %s"), *SourceName.ToString(), *DestName.ToString(), LAZY_COPY_DECORATOR);
 			break;
 		}
 
 		// Copy from buffer
 		case D3D12_TEXTURE_COPY_TYPE_PLACED_FOOTPRINT:
 		{
-			const uint32 RequiredPlacedSize = GetPlacedBufferRequiredSize(*this, SourceCopyLocation->PlacedFootprint.Footprint, SourcePixelFormatInfo, DebugName);
-			RHI_BREADCRUMB_CHECK_SHIPPINGF(*this, SourceCopyLocation->PlacedFootprint.Offset + RequiredPlacedSize <= SourceDesc.Width, TEXT("Source placed buffer width out of bounds for: '%s'"), *DebugName.ToString());
+			const uint32 RequiredPlacedSize = GetPlacedBufferRequiredSize(*this, SourceCopyLocation->PlacedFootprint.Footprint, SourcePixelFormatInfo, SourceName);
+			RHI_BREADCRUMB_CHECK_SHIPPINGF(*this, SourceCopyLocation->PlacedFootprint.Offset + RequiredPlacedSize <= SourceDesc.Width, TEXT("Source placed buffer width out of bounds for: '%s' -> '%s'"), *SourceName.ToString(), *DestName.ToString());
 			break;
 		}
 	}
@@ -1575,20 +1575,20 @@ FORCENOINLINE void FD3D12CommandContext::CopyTextureRegionChecked(
 			const uint32 SubresourceHeight = AlignArbitrary<uint32>(FMath::Max<uint32>(1u, DestDesc.Height >> Indices.MipIndex), DestPixelFormatInfo.BlockSizeY);
 			const uint32 SubresourceDepth  = AlignArbitrary<uint32>(DestDesc.Dimension == D3D12_RESOURCE_DIMENSION_TEXTURE3D ? FMath::Max<uint32>(1u, DestDesc.DepthOrArraySize >> Indices.MipIndex) : 1, DestPixelFormatInfo.BlockSizeZ);
 
-			RHI_BREADCRUMB_CHECK_SHIPPINGF(*this, DestX % DestPixelFormatInfo.BlockSizeX == 0 && (CopyBox.right - CopyBox.left) % DestPixelFormatInfo.BlockSizeX == 0, TEXT("Width not aligned to block size for: '%s', %s"), *DebugName.ToString(), LAZY_COPY_DECORATOR);
-			RHI_BREADCRUMB_CHECK_SHIPPINGF(*this, DestY % DestPixelFormatInfo.BlockSizeY == 0 && (CopyBox.bottom - CopyBox.top) % DestPixelFormatInfo.BlockSizeY == 0, TEXT("Height not aligned to block size for: '%s', %s"), *DebugName.ToString(), LAZY_COPY_DECORATOR);
+			RHI_BREADCRUMB_CHECK_SHIPPINGF(*this, DestX % DestPixelFormatInfo.BlockSizeX == 0 && (CopyBox.right - CopyBox.left) % DestPixelFormatInfo.BlockSizeX == 0, TEXT("Width not aligned to block size for: '%s' <- '%s', %s"), *DestName.ToString(), *SourceName.ToString(), LAZY_COPY_DECORATOR);
+			RHI_BREADCRUMB_CHECK_SHIPPINGF(*this, DestY % DestPixelFormatInfo.BlockSizeY == 0 && (CopyBox.bottom - CopyBox.top) % DestPixelFormatInfo.BlockSizeY == 0, TEXT("Height not aligned to block size for: '%s' <- '%s', %s"), *DestName.ToString(), *SourceName.ToString(), LAZY_COPY_DECORATOR);
 				
-			RHI_BREADCRUMB_CHECK_SHIPPINGF(*this, DestX + (CopyBox.right - CopyBox.left) <= SubresourceWidth, TEXT("Dest width out of bounds for: '%s', %s"), *DebugName.ToString(), LAZY_COPY_DECORATOR);
-			RHI_BREADCRUMB_CHECK_SHIPPINGF(*this, DestY + (CopyBox.bottom - CopyBox.top) <= SubresourceHeight, TEXT("Dest height out of bounds for: '%s', %s"), *DebugName.ToString(), LAZY_COPY_DECORATOR);
-			RHI_BREADCRUMB_CHECK_SHIPPINGF(*this, DestZ + (CopyBox.back - CopyBox.front) <= SubresourceDepth, TEXT("Dest depth out of bounds for: '%s', %s"), *DebugName.ToString(), LAZY_COPY_DECORATOR);
+			RHI_BREADCRUMB_CHECK_SHIPPINGF(*this, DestX + (CopyBox.right - CopyBox.left) <= SubresourceWidth, TEXT("Dest width out of bounds for: '%s' <- '%s', %s"), *DestName.ToString(), *SourceName.ToString(), LAZY_COPY_DECORATOR);
+			RHI_BREADCRUMB_CHECK_SHIPPINGF(*this, DestY + (CopyBox.bottom - CopyBox.top) <= SubresourceHeight, TEXT("Dest height out of bounds for: '%s' <- '%s', %s"), *DestName.ToString(), *SourceName.ToString(), LAZY_COPY_DECORATOR);
+			RHI_BREADCRUMB_CHECK_SHIPPINGF(*this, DestZ + (CopyBox.back - CopyBox.front) <= SubresourceDepth, TEXT("Dest depth out of bounds for: '%s' <- '%s', %s"), *DestName.ToString(), *SourceName.ToString(), LAZY_COPY_DECORATOR);
 			break;
 		}
 
 		// Copy to buffer
 		case D3D12_TEXTURE_COPY_TYPE_PLACED_FOOTPRINT:
 		{
-			const uint32 RequiredPlacedSize = GetPlacedBufferRequiredSize(*this, DestCopyLocation->PlacedFootprint.Footprint, DestPixelFormatInfo, DebugName);
-			RHI_BREADCRUMB_CHECK_SHIPPINGF(*this, DestCopyLocation->PlacedFootprint.Offset + RequiredPlacedSize <= DestDesc.Width, TEXT("Dest placed buffer width out of bounds for: '%s'"), *DebugName.ToString());
+			const uint32 RequiredPlacedSize = GetPlacedBufferRequiredSize(*this, DestCopyLocation->PlacedFootprint.Footprint, DestPixelFormatInfo, DestName);
+			RHI_BREADCRUMB_CHECK_SHIPPINGF(*this, DestCopyLocation->PlacedFootprint.Offset + RequiredPlacedSize <= DestDesc.Width, TEXT("Dest placed buffer width out of bounds for: '%s' <- '%s'"), *DestName.ToString(), *SourceName.ToString());
 			break;
 		}
 	}
@@ -2064,7 +2064,8 @@ FRHILockTextureResult FD3D12Texture::Lock(FRHICommandListBase& RHICmdList, const
 				&SourceCopyLocation,
 				nullptr,
 				GetFormat(),
-				GetName()
+				StagingTexture->GetName(),		// DestName
+				GetName()						// SourceName
 			);
 
 			Context.UpdateResidency(StagingTexture);
@@ -2109,7 +2110,8 @@ void FD3D12Texture::UpdateTexture(FD3D12CommandContext& Context, uint32 MipIndex
 		&SourceCopyLocation,
 		nullptr,
 		GetFormat(),
-		GetName()
+		GetName(),						// DestName
+		SourceResource->GetName()		// SourceName
 	);
 
 	Context.UpdateResidency(GetResource());
@@ -2145,7 +2147,8 @@ void FD3D12Texture::CopyTextureRegion(FD3D12CommandContext& Context, uint32 Dest
 		&SourceCopyLocation,
 		&SourceBox,
 		SourceTexture->GetFormat(),
-		GetName()
+		GetName(),						// DestName
+		SourceTexture->GetName()		// SourceName
 	);
 
 	Context.UpdateResidency(SourceTexture->GetResource());
@@ -2200,9 +2203,11 @@ void FD3D12Texture::UploadInitialData(
 
 			for (uint32 Subresource = 0; Subresource < NumSubresources; Subresource++)
 			{
+				// Src in this case is an anonymous resource location, named UploadLocation in the calling code, so use that name here.
+				const static FName UploadLocationName("UploadLocation");
 				Dst.SubresourceIndex = Subresource;
 				Src.PlacedFootprint = Footprints[Subresource];
-				Context.CopyTextureRegionChecked(&Dst, 0, 0, 0, Texture->GetFormat(), &Src, nullptr, Texture->GetFormat(), Texture->GetName());
+				Context.CopyTextureRegionChecked(&Dst, 0, 0, 0, Texture->GetFormat(), &Src, nullptr, Texture->GetFormat(), Texture->GetName(), UploadLocationName);
 			}
 
 			Context.UpdateResidency(SourceLocation.GetResource());
@@ -2554,7 +2559,8 @@ public:
 					&SourceCopyLocation,
 					nullptr,
 					TextureLink.GetFormat(),
-					TextureLink.GetName()
+					TextureLink.GetName(),		// DestName
+					UploadBuffer->GetName()		// SourceName
 				);
 
 				Context.UpdateResidency(TextureLink.GetResource());
@@ -2795,7 +2801,8 @@ void FD3D12DynamicRHI::EndUpdateTexture3D_Internal(FRHICommandListBase& RHICmdLi
 					&SourceCopyLocation,
 					nullptr,
 					Texture.GetFormat(),
-					Texture.GetName()
+					Texture.GetName(),									// DestName
+					SrcResourceLocation->GetResource()->GetName()		// SourceName
 				);
 
 				Context.ConditionalSplitCommandList();
@@ -3121,7 +3128,8 @@ void FD3D12CommandContext::RHICopyTexture(FRHITexture* SourceTextureRHI, FRHITex
 					&Src,
 					&SrcBox,
 					SourceTexture->GetFormat(),
-					SourceTexture->GetName()
+					DestTexture->GetName(),		// DestName
+					SourceTexture->GetName()	// SourceName
 				);
 			}
 		}
