Add -wavesizelimit command line option, allows limiting max wave size for testing.
#Synced-CL 52493164
#Synced-CL 52493164

diff --git a/Engine/Source/Runtime/D3D12RHI/Private/D3D12Adapter.cpp b/Engine/Source/Runtime/D3D12RHI/Private/D3D12Adapter.cpp
--- a/Engine/Source/Runtime/D3D12RHI/Private/D3D12Adapter.cpp
+++ b/Engine/Source/Runtime/D3D12RHI/Private/D3D12Adapter.cpp
@@ -1195,10 +1195,29 @@ void FD3D12Adapter::InitializeDevices()
 				GRHISupportsWaveOperations = Features.WaveOps;
 				GRHIMinimumWaveSize = Features.WaveLaneCountMin;
 				GRHIMaximumWaveSize = Features.WaveLaneCountMax;
-
 				if (GRHISupportsWaveOperations)
 				{
 					UE_LOG(LogD3D12RHI, Log, TEXT("Wave Operations are supported (wave size: min=%d max=%d)."), GRHIMinimumWaveSize, GRHIMaximumWaveSize);
+					int32 MaximumWaveSizeLimit;
+					if (FParse::Value(FCommandLine::Get(), TEXT("-wavesizelimit="), MaximumWaveSizeLimit))
+					{
+						if (MaximumWaveSizeLimit > GRHIMinimumWaveSize)
+						{
+							if (MaximumWaveSizeLimit < GRHIMaximumWaveSize)
+							{
+								UE_LOGF(LogD3D12RHI, Log, "Wave Operations limited from %u to %d by -wavesizelimit", Features.WaveLaneCountMax, MaximumWaveSizeLimit);
+								GRHIMaximumWaveSize = MaximumWaveSizeLimit;
+							}
+							else
+							{
+								UE_LOGF(LogD3D12RHI, Warning, "Device WaveLaneCountMax %u is less than or equal to -wavesizelimit of %d so the limit was not applied.", Features.WaveLaneCountMax, MaximumWaveSizeLimit);
+							}
+						}
+						else
+						{
+							UE_LOGF(LogD3D12RHI, Warning, "Device WaveLaneCountMin %u is greater than or equal to -wavesizelimit of %d so the limit was not applied.", Features.WaveLaneCountMin, MaximumWaveSizeLimit);
+						}
+					}
 				}
 			}
 
