[D3D12RHI] Fix for scissor rect validation failing with fractional viewports.

#Synced-CL 52642033

diff --git a/Engine/Source/Runtime/D3D12RHI/Private/D3D12Commands.cpp b/Engine/Source/Runtime/D3D12RHI/Private/D3D12Commands.cpp
--- a/Engine/Source/Runtime/D3D12RHI/Private/D3D12Commands.cpp
+++ b/Engine/Source/Runtime/D3D12RHI/Private/D3D12Commands.cpp
@@ -302,7 +302,7 @@ void FD3D12CommandContext::RHISetViewport(float MinX, float MinY, float MinZ, fl
 	{
 		// Setting a viewport will also set the scissor rect appropriately.
 		StateCache.SetViewport(Viewport);
-		RHISetScissorRect(true, MinX, MinY, MaxX, MaxY);
+		RHISetScissorRect(true, FMath::RoundToZero(MinX), FMath::RoundToZero(MinY), FMath::RoundFromZero(MaxX), FMath::RoundFromZero(MaxY));
 	}
 }
 
diff --git a/Engine/Source/Runtime/D3D12RHI/Private/D3D12StateCache.cpp b/Engine/Source/Runtime/D3D12RHI/Private/D3D12StateCache.cpp
--- a/Engine/Source/Runtime/D3D12RHI/Private/D3D12StateCache.cpp
+++ b/Engine/Source/Runtime/D3D12RHI/Private/D3D12StateCache.cpp
@@ -304,22 +304,22 @@ void FD3D12StateCache::SetViewports(uint32 Count, const D3D12_VIEWPORT* const Vi
 static void ValidateScissorRect(const D3D12_VIEWPORT& Viewport, const D3D12_RECT& ScissorRect)
 {
 	bool bScissorRectValid = true;
-	bScissorRectValid = bScissorRectValid && ScissorRect.left >= (LONG)Viewport.TopLeftX;
-	bScissorRectValid = bScissorRectValid && ScissorRect.top >= (LONG)Viewport.TopLeftY;
-	bScissorRectValid = bScissorRectValid && ScissorRect.right <= (LONG)Viewport.TopLeftX + (LONG)Viewport.Width;
-	bScissorRectValid = bScissorRectValid && ScissorRect.bottom <= (LONG)Viewport.TopLeftY + (LONG)Viewport.Height;
+	bScissorRectValid = bScissorRectValid && ScissorRect.left >= (LONG)FMath::RoundToZero(Viewport.TopLeftX);
+	bScissorRectValid = bScissorRectValid && ScissorRect.top >= (LONG)FMath::RoundToZero(Viewport.TopLeftY);
+	bScissorRectValid = bScissorRectValid && ScissorRect.right <= (LONG)FMath::RoundFromZero(Viewport.TopLeftX + Viewport.Width);
+	bScissorRectValid = bScissorRectValid && ScissorRect.bottom <= (LONG)FMath::RoundFromZero(Viewport.TopLeftY + Viewport.Height);
 	bScissorRectValid = bScissorRectValid && ScissorRect.left <= ScissorRect.right && ScissorRect.top <= ScissorRect.bottom;
 
 	ensureMsgf(bScissorRectValid,
-		TEXT("Scissor invalid with current Viewport. Scissor: [left:%li, top:%li, right:%li, bottom:%li]. Viewport: [left:%li, top:%li, right:%li, bottom:%li]")
+		TEXT("Scissor invalid with current Viewport. Scissor: [left:%li, top:%li, right:%li, bottom:%li]. Viewport: [left:%.2f, top:%.2f, right:%.2f, bottom:%.2f]")
 			, ScissorRect.left
 			, ScissorRect.top
 			, ScissorRect.right
 			, ScissorRect.bottom
-			, (LONG)Viewport.TopLeftX
-			, (LONG)Viewport.TopLeftY
-			, (LONG)Viewport.TopLeftX + (LONG)Viewport.Width
-			, (LONG)Viewport.TopLeftY + (LONG)Viewport.Height);
+			, Viewport.TopLeftX
+			, Viewport.TopLeftY
+			, Viewport.TopLeftX + Viewport.Width
+			, Viewport.TopLeftY + Viewport.Height);
 }
 
 void FD3D12StateCache::SetScissorRect(const D3D12_RECT& ScissorRect)
